[QuicMessaging] Moving QuicMessaging to experimental

Moving QuicMessaging to the experimental plugins folder.

diff --git a/Engine/Plugins/Experimental/QuicMessaging/QuicMessaging.uplugin b/Engine/Plugins/Experimental/QuicMessaging/QuicMessaging.uplugin
new file mode 100644
--- /dev/null
+++ b/Engine/Plugins/Experimental/QuicMessaging/QuicMessaging.uplugin
@@ -0,0 +1,69 @@
+{
+	"FileVersion": 3,
+	"Version": 1,
+	"VersionName": "1.0",
+	"FriendlyName": "QUIC Messaging",
+	"Description": "Adds a QUIC based transport layer to the messaging sub-system for sending and receiving messages between networked computers and devices.",
+	"Category": "Messaging",
+	"CreatedBy": "Epic Games, Inc.",
+	"CreatedByURL": "http://epicgames.com",
+	"DocsURL": "",
+	"MarketplaceURL": "",
+	"SupportURL": "",
+	"EnabledByDefault": false,
+	"CanContainContent": false,
+	"IsBetaVersion": false,
+	"Installed": false,
+	"SupportedTargetPlatforms": [
+		"Win64",
+		"Linux",
+		"Mac"
+	],
+	"SupportedPrograms": [
+		"UnrealFrontend",
+		"UnrealMultiUserServer",
+		"UnrealMultiUserSlateServer",
+		"UnrealRecoverySvc",
+		"UnrealLightmass",
+		"CrashReportClientEditor",
+		"CoopMultiUserServer"
+	],
+	"Modules": [
+		{
+			"Name": "QuicMessaging",
+			"Type": "RuntimeAndProgram",
+			"LoadingPhase": "PreDefault",
+			"PlatformAllowList": [ "Win64", "Linux", "Mac" ],
+			"ProgramAllowList": [
+				"UnrealFrontend",
+				"UnrealMultiUserServer",
+				"UnrealMultiUserSlateServer",
+				"UnrealRecoverySvc",
+				"UnrealLightmass",
+				"CrashReportClientEditor",
+				"CoopMultiUserServer"
+			]
+		},
+		{
+			"Name": "QuicMessagingTransport",
+			"Type": "RuntimeAndProgram",
+			"LoadingPhase": "PreDefault",
+			"PlatformAllowList": [ "Win64", "Linux", "Mac" ],
+			"ProgramAllowList": [
+				"UnrealFrontend",
+				"UnrealMultiUserServer",
+				"UnrealMultiUserSlateServer",
+				"UnrealRecoverySvc",
+				"UnrealLightmass",
+				"CrashReportClientEditor",
+				"CoopMultiUserServer"
+			]
+		}
+	],
+	"Plugins": [
+		{
+			"Name": "MsQuic",
+			"Enabled": true
+		}
+	]
+}
diff --git a/Engine/Plugins/Experimental/QuicMessaging/Source/QuicMessaging/Private/QuicMessagingModule.cpp b/Engine/Plugins/Experimental/QuicMessaging/Source/QuicMessaging/Private/QuicMessagingModule.cpp
new file mode 100644
--- /dev/null
+++ b/Engine/Plugins/Experimental/QuicMessaging/Source/QuicMessaging/Private/QuicMessagingModule.cpp
@@ -0,0 +1,950 @@
+// Copyright Epic Games, Inc. All Rights Reserved.
+
+#include "QuicMessagingPrivate.h"
+#include "QuicMessageProcessor.h"
+#include "QuicEndpointConfig.h"
+
+#include "CoreMinimal.h"
+#include "CoreTypes.h"
+#include "MessageBridgeBuilder.h"
+#include "Misc/CoreMisc.h"
+#include "Misc/CommandLine.h"
+#include "Misc/CoreDelegates.h"
+#include "Misc/App.h"
+#include "Misc/Paths.h"
+#include "Containers/Ticker.h"
+#include "Interfaces/IPv4/IPv4Endpoint.h"
+#include "Stats/Stats.h"
+
+#if WITH_EDITOR
+	#include "ISettingsModule.h"
+	#include "ISettingsSection.h"
+#endif
+
+#include "MsQuicRuntimeModule.h"
+#include "Features/IModularFeatures.h"
+#include "IQuicNetworkMessagingExtension.h"
+#include "MessageEndpoint.h"
+#include "QuicTransportMessages.h"
+#include "Interfaces/IPluginManager.h"
+#include "Shared/QuicMessagingSettings.h"
+#include "Transport/QuicMessageTransport.h"
+
+DEFINE_LOG_CATEGORY(LogQuicMessaging);
+
+#define LOCTEXT_NAMESPACE "FQuicMessagingModule"
+
+
+/**
+ * Implements the QuicMessagingModule and the network messaging extension modular feature.
+ */
+class FQuicMessagingModule
+    : public FSelfRegisteringExec
+    , public IQuicNetworkMessagingExtension
+    , public IModuleInterface
+{
+public:
+
+    //~ FSelfRegisteringExec interface
+
+    virtual bool Exec(UWorld* InWorld, const TCHAR* Cmd, FOutputDevice& Ar) override
+    {
+        if (!FParse::Command(&Cmd, TEXT("QUICMESSAGING")))
+		{
+			return false;
+		}
+
+        if (FParse::Command(&Cmd, TEXT("STATUS")))
+		{
+			UQuicMessagingSettings* Settings = GetMutableDefault<UQuicMessagingSettings>();
+
+			// Bridge status
+			if (MessageBridge.IsValid())
+			{
+				if (MessageBridge->IsEnabled())
+				{
+					Ar.Log(TEXT("Message Bridge: Initialized and enabled"));
+				}
+				else
+				{
+					Ar.Log(TEXT("Message Bridge: Initialized, but disabled"));
+				}
+			}
+			else
+			{
+				Ar.Log(TEXT("Message Bridge: Not initialized."));
+			}
+
+			// Bridge settings
+			if (Settings->UnicastEndpoint.IsEmpty())
+			{
+				Ar.Logf(TEXT("    Unicast Endpoint: %s (default)"),
+					*FIPv4Endpoint::Any.ToString());
+			}
+			else
+			{
+				Ar.Logf(TEXT("    Unicast Endpoint: %s"), *Settings->UnicastEndpoint);
+			}
+
+			if (Settings->StaticEndpoints.Num() > 0)
+			{
+				Ar.Log(TEXT("    Static Endpoints:"));
+
+				for (const auto& StaticEndpoint : Settings->StaticEndpoints)
+				{
+					Ar.Logf(TEXT("        %s"), *StaticEndpoint);
+				}
+			}
+			else
+			{
+				Ar.Log(TEXT("    Static Endpoints: None"));
+			}
+        }
+        else if (FParse::Command(&Cmd, TEXT("RESTART")))
+        {
+            RestartServices();
+        }
+        else if (FParse::Command(&Cmd, TEXT("SHUTDOWN")))
+        {
+            ShutdownBridge();
+        }
+        else
+        {
+			// Show usage
+			Ar.Log(TEXT("Usage: QUICMESSAGING <Command>"));
+			Ar.Log(TEXT(""));
+			Ar.Log(TEXT("Command"));
+			Ar.Log(TEXT("    RESTART = Restarts the message bridge, if enabled"));
+			Ar.Log(TEXT("    SHUTDOWN = Shut down the message bridge, if running"));
+			Ar.Log(TEXT("    STATUS = Displays the status of the QUIC message transport"));
+        }
+
+        return true;
+    }
+
+public:
+
+    // IModuleInterface
+    virtual void StartupModule() override
+    {
+		if (!IsSupportEnabled())
+		{
+			return;
+		}
+
+        // Load networking dependency
+		if (!FModuleManager::Get().LoadModule(TEXT("Networking")))
+        {
+            UE_LOG(LogQuicMessaging, Error,
+				TEXT("The required module 'Networking' failed to load."
+				"Plug-in 'QUIC Messaging' cannot be used."));
+
+			return;
+        }
+
+		// Initiate MsQuicRuntime (loads the DLLs)
+		if (!FMsQuicRuntimeModule::InitRuntime())
+		{
+			UE_LOG(LogQuicMessaging, Error,
+				TEXT("[QuicMessagingModule] Could not initialize MsQuicRuntime."));
+
+			return;
+		}
+
+		// Hook to the PreExit callback, needed to execute UObject related shutdowns
+		FCoreDelegates::OnPreExit.AddRaw(
+			this, &FQuicMessagingModule::HandleAppPreExit);
+
+#if WITH_EDITOR
+    	
+        // Register settings
+        if (ISettingsModule* SettingsModule
+			= FModuleManager::GetModulePtr<ISettingsModule>("Settings"))
+        {
+            const ISettingsSectionPtr SettingsSection
+        		= SettingsModule->RegisterSettings("Project", "Plugins", "QuicMessaging",
+                LOCTEXT("QuicMessagingSettingsName", "QUIC Messaging"),
+                LOCTEXT("QuicMessagingSettingsDescription", "Configure the QUIC Messaging plugin."),
+                GetMutableDefault<UQuicMessagingSettings>()
+            );
+
+            if (SettingsSection.IsValid())
+            {
+                SettingsSection->OnModified().BindRaw(
+					this, &FQuicMessagingModule::HandleSettingsSaved);
+            }
+        }
+#endif
+
+        // Parse additional command line args
+        ParseCommandLine(GetMutableDefault<UQuicMessagingSettings>(), FCommandLine::Get());
+
+        // Register application events
+        const UQuicMessagingSettings& Settings = *GetDefault<UQuicMessagingSettings>();
+
+        if (Settings.bStopServiceWhenAppDeactivates)
+        {
+            FCoreDelegates::ApplicationHasReactivatedDelegate.AddRaw(
+				this, &FQuicMessagingModule::HandleApplicationHasReactivated);
+
+			FCoreDelegates::ApplicationWillDeactivateDelegate.AddRaw(
+				this, &FQuicMessagingModule::HandleApplicationWillDeactivate);
+        }
+
+		// Create the endpoints Guid
+    	CreateEndpointGuid();
+
+		// Restart services to apply changes
+        RestartServices();
+
+        IModularFeatures::Get().RegisterModularFeature(
+			INetworkMessagingExtension::ModularFeatureName, this);
+    }
+
+    virtual void ShutdownModule() override
+    {
+        // Un-register application events
+		FCoreDelegates::ApplicationHasReactivatedDelegate.RemoveAll(this);
+		FCoreDelegates::ApplicationWillDeactivateDelegate.RemoveAll(this);
+
+		// Unhook AppPreExit and call it
+		FCoreDelegates::OnPreExit.RemoveAll(this);
+		HandleAppPreExit();
+
+#if WITH_EDITOR
+        // Un-register settings
+
+		if (ISettingsModule* SettingsModule
+			= FModuleManager::GetModulePtr<ISettingsModule>("Settings"))
+		{
+			SettingsModule->UnregisterSettings(
+				"Project", "Plugins", "QuicMessaging");
+		}
+#endif
+
+        // Shut down services
+        ShutdownBridge();
+
+        IModularFeatures::Get().UnregisterModularFeature(
+			INetworkMessagingExtension::ModularFeatureName, this);
+    }
+
+    virtual bool SupportsDynamicReloading() override
+    {
+        return true;
+    }
+
+    // INetworkMessagingExtension interface
+    virtual FName GetName() const override
+    {
+        return QuicMessagingName;
+    }
+
+    virtual bool IsSupportEnabled() const override
+    {
+#if !IS_PROGRAM && UE_BUILD_SHIPPING && !(defined(ALLOW_QUIC_MESSAGING_SHIPPING) && ALLOW_QUIC_MESSAGING_SHIPPING)
+        return false;
+#else
+
+    	// Disallow unsupported platforms
+		if (!FPlatformMisc::SupportsMessaging())
+		{
+			return false;
+		}
+
+		// Always allow in standalone Slate applications
+		if (!FApp::IsGame() && !IsRunningCommandlet())
+		{
+			return true;
+		}
+
+		// Otherwise allow if explicitly desired
+		if (FParse::Param(FCommandLine::Get(), TEXT("Messaging")))
+		{
+			return true;
+		}
+
+        // Check the project setting
+        const UQuicMessagingSettings& Settings = *GetDefault<UQuicMessagingSettings>();
+        return Settings.EnabledByDefault;
+#endif
+    }
+
+    virtual void RestartServices() override
+    {
+        const UQuicMessagingSettings& Settings = *GetDefault<UQuicMessagingSettings>();
+
+        if (Settings.EnableTransport)
+        {
+            InitializeBridge();
+        }
+        else
+        {
+            ShutdownBridge();
+        }
+    }
+
+    virtual void ShutdownServices() override
+    {
+        ShutdownBridge();
+
+        AdditionalStaticEndpoints.Empty();
+    }
+
+	virtual bool CanProvideNetworkStatistics() const override
+	{
+		return true;
+	}
+
+	virtual FMessageTransportStatistics GetLatestNetworkStatistics(FGuid NodeId) const override
+    {
+		if (const TSharedPtr<FQuicMessageTransport> Transport = WeakBridgeTransport.Pin())
+		{
+			return Transport->GetLatestStatistics(NodeId);
+		}
+
+		return {};
+    }
+
+	virtual FGuid GetNodeIdFromAddress(const FMessageAddress& MessageAddress) const override
+    {
+		return MessageBridge->LookupAddress(MessageAddress);
+    }
+
+	virtual FOnOutboundTransferDataUpdated& OnOutboundTransferUpdatedFromThread() override
+	{
+		return UE::Private::MessageProcessor::OnOutboundUpdated();
+	}
+
+	virtual FOnInboundTransferDataUpdated& OnInboundTransferUpdatedFromThread() override
+	{
+		return UE::Private::MessageProcessor::OnInboundUpdated();
+	}
+
+	virtual TArray<FString> GetListeningAddresses() const override
+    {
+		if (const TSharedPtr<FQuicMessageTransport> Transport = WeakBridgeTransport.Pin())
+		{
+			TArray<FString> StringifiedEndpoints;
+			TArray<FIPv4Endpoint> Endpoints = Transport->GetListeningAddresses();
+
+			for (const FIPv4Endpoint& Endpoint : Endpoints)
+			{
+				StringifiedEndpoints.Add(Endpoint.ToString());
+			}
+
+			return StringifiedEndpoints;
+		}
+
+    	return {};
+    }
+
+	virtual TArray<FString> GetKnownEndpoints() const override
+    {
+		if (const TSharedPtr<FQuicMessageTransport> Transport = WeakBridgeTransport.Pin())
+		{
+			TArray<FString> StringifiedEndpoints;
+			TArray<FIPv4Endpoint> Endpoints = Transport->GetKnownEndpoints();
+
+			for (const FIPv4Endpoint& Endpoint : Endpoints)
+			{
+				StringifiedEndpoints.Add(Endpoint.ToString());
+			}
+
+			return StringifiedEndpoints;
+		}
+
+    	return {};
+    }
+
+    virtual void AddEndpoint(const FString& InEndpoint) override
+    {
+        if (const TSharedPtr<FQuicMessageTransport> Transport = WeakBridgeTransport.Pin())
+        {
+			if (InEndpoint == "")
+			{
+				return;
+			}
+
+            FScopeLock StaticEndpointsLock(&StaticEndpointsCS);
+            FIPv4Endpoint OutEndpoint;
+        	
+            if (ParseEndpoint(InEndpoint, OutEndpoint)
+				&& !AdditionalStaticEndpoints.Contains(OutEndpoint))
+            {	
+                AdditionalStaticEndpoints.Add(OutEndpoint);
+                Transport->AddStaticEndpoint(OutEndpoint);
+            }
+        }
+    }
+
+	virtual void RemoveEndpoint(const FString& InEndpoint) override
+	{
+		if (const TSharedPtr<FQuicMessageTransport> Transport = WeakBridgeTransport.Pin())
+		{
+			FScopeLock StaticEndpointsLock(&StaticEndpointsCS);
+			FIPv4Endpoint OutEndpoint;
+
+			if (ParseEndpoint(InEndpoint, OutEndpoint)
+				&& AdditionalStaticEndpoints.Contains(OutEndpoint))
+			{
+				AdditionalStaticEndpoints.Remove(OutEndpoint);
+				Transport->RemoveStaticEndpoint(OutEndpoint);
+			}
+		}
+	}
+
+public:
+
+	/** IQuicNetworkMessagingExtension implementation. */
+
+	virtual FGuid GetEndpointGuid() override
+    {
+		const UQuicMessagingSettings& Settings = *GetDefault<UQuicMessagingSettings>();
+
+    	return Settings.EndpointGuid;
+    }
+
+	virtual TOptional<FGuid> GetNodeId(const FString& RemoteEndpoint) override
+	{
+		const TSharedPtr<FQuicMessageTransport> Transport = WeakBridgeTransport.Pin();
+
+		if (!Transport.IsValid())
+		{
+			UE_LOG(LogQuicMessaging, Error,
+				TEXT("[QuicMessagingAPI::GetNodeId] Transport pointer not valid."));
+
+			return TOptional<FGuid>();
+		}
+
+		FIPv4Endpoint ParsedRemoteEndpoint;
+
+		if (!FIPv4Endpoint::Parse(RemoteEndpoint, ParsedRemoteEndpoint))
+		{
+			UE_LOG(LogQuicMessaging, Error,
+				TEXT("[QuicMessagingAPI::GetNodeId] Endpoint string could not be parsed."));
+
+			return TOptional<FGuid>();
+		}
+
+		return Transport->GetNodeId(ParsedRemoteEndpoint);
+	}
+
+	virtual void SetMaxAuthenticationMessageSize(const uint32 MaxBytes) override
+	{
+		const TSharedPtr<FQuicMessageTransport> Transport = WeakBridgeTransport.Pin();
+
+		if (!Transport.IsValid())
+		{
+			UE_LOG(LogQuicMessaging, Error,
+				TEXT("[QuicMessagingAPI::SetMaxAuthenticationMessageSize] "
+				"Transport pointer not valid."));
+
+			return;
+		}
+
+		Transport->SetMaxAuthenticationMessageSize(MaxBytes);
+	}
+
+	virtual bool IsNodeAuthenticated(const FGuid& NodeId) const override
+	{
+		const TSharedPtr<FQuicMessageTransport> Transport = WeakBridgeTransport.Pin();
+
+		if (!Transport.IsValid())
+		{
+			UE_LOG(LogQuicMessaging, Error,
+				TEXT("[QuicMessagingAPI::IsNodeAuthenticated] Transport pointer not valid."));
+
+			return false;
+		}
+
+		return Transport->IsNodeAuthenticated(NodeId);
+	}
+
+	virtual void SetNodeAuthenticated(const FGuid& NodeId) override
+	{
+		const TSharedPtr<FQuicMessageTransport> Transport = WeakBridgeTransport.Pin();
+
+		if (!Transport.IsValid())
+		{
+			UE_LOG(LogQuicMessaging, Error,
+				TEXT("[QuicMessagingAPI::SetNodeAuthenticated] Transport pointer not valid."));
+
+			return;
+		}
+
+		Transport->SetNodeAuthenticated(NodeId);
+	}
+
+	virtual void DisconnectNode(const FGuid& NodeId) override
+	{
+		const TSharedPtr<FQuicMessageTransport> Transport = WeakBridgeTransport.Pin();
+
+		if (!Transport.IsValid())
+		{
+			UE_LOG(LogQuicMessaging, Error,
+				TEXT("[QuicMessagingAPI::DisconnectNode] Transport pointer not valid."));
+
+			return;
+		}
+
+		Transport->DisconnectNode(NodeId);
+	}
+
+	virtual void SetConnectionCooldown(
+		const bool bEnabled, const uint32 MaxAttempts, const uint32 PeriodSeconds,
+		const uint32 CooldownSeconds, const uint32 CooldownMaxSeconds) override
+	{
+		const TSharedPtr<FQuicMessageTransport> Transport = WeakBridgeTransport.Pin();
+
+		if (!Transport.IsValid())
+		{
+			UE_LOG(LogQuicMessaging, Error,
+				TEXT("[QuicMessagingAPI::SetConnectionCooldown] Transport pointer not valid."));
+
+			return;
+		}
+
+		Transport->SetConnectionCooldown(
+			bEnabled, MaxAttempts, PeriodSeconds, CooldownSeconds, CooldownMaxSeconds);
+	}
+
+	virtual bool TransportAuthMessage(
+		const TSharedRef<IMessageContext, ESPMode::ThreadSafe>& Context,
+    	const FGuid& Recipient) override
+	{
+		const TSharedPtr<FQuicMessageTransport> Transport = WeakBridgeTransport.Pin();
+
+		if (!Transport.IsValid())
+		{
+			UE_LOG(LogQuicMessaging, Error,
+				TEXT("[QuicMessagingAPI::TransportAuthMessage] Transport pointer not valid."));
+
+			return false;
+		}
+
+		return Transport->TransportAuthMessage(Context, Recipient);
+	}
+
+	virtual bool TransportAuthResponseMessage(
+		const TSharedRef<IMessageContext, ESPMode::ThreadSafe>& Context,
+    	const FGuid& Recipient) override
+	{
+		const TSharedPtr<FQuicMessageTransport> Transport = WeakBridgeTransport.Pin();
+
+		if (!Transport.IsValid())
+		{
+			UE_LOG(LogQuicMessaging, Error,
+				TEXT("[QuicMessagingAPI::TransportAuthResponseMessage]"
+					"Transport pointer not valid."));
+
+			return false;
+		}
+
+		return Transport->TransportAuthResponseMessage(Context, Recipient);
+	}
+
+	virtual FOnQuicMetaMessageReceived& OnQuicMetaMessageReceived() override
+	{
+		const TSharedPtr<FQuicMessageTransport> Transport = WeakBridgeTransport.Pin();
+
+		if (!Transport.IsValid())
+		{
+			UE_LOG(LogQuicMessaging, Fatal,
+				TEXT("[QuicMessagingAPI::OnQuicMetaMessageReceived] Transport pointer not valid."));
+		}
+
+		return Transport->OnMetaMessageReceived();
+	}
+
+protected:
+
+	static bool ParseEndpoint(const FString& InEndpointString,
+		FIPv4Endpoint& OutEndpoint, bool bSupportHostname = true)
+	{
+		bool bParsedEndpoint = FIPv4Endpoint::Parse(InEndpointString, OutEndpoint);
+
+		if (bSupportHostname)
+		{
+			bParsedEndpoint = FIPv4Endpoint::FromHostAndPort(InEndpointString, OutEndpoint);
+		}
+
+		return bParsedEndpoint;
+	}
+
+    /** Initializes the message bridge with the current settings. */
+    void InitializeBridge()
+    {
+        ShutdownBridge();
+    	
+        UQuicMessagingSettings* Settings = GetMutableDefault<UQuicMessagingSettings>();
+        bool bResaveSettings = false;
+
+		FIPv4Endpoint UnicastEndpoint;
+
+		if (!ParseEndpoint(Settings->UnicastEndpoint, UnicastEndpoint))
+		{
+			if (!Settings->UnicastEndpoint.IsEmpty())
+			{
+				UE_LOG(LogQuicMessaging, Warning,
+					TEXT("Invalid setting for UnicastEndpoint '%s' - "
+					"binding to all local network adapters instead"),
+					*Settings->UnicastEndpoint);
+			}
+
+			UnicastEndpoint = FIPv4Endpoint::Any;
+			Settings->UnicastEndpoint = UnicastEndpoint.ToText().ToString();
+			bResaveSettings = true;
+		}
+
+        // Initialize the service with the additional endpoints added through the modular interface
+		TArray<FIPv4Endpoint> StaticEndpoints = AdditionalStaticEndpoints.Array();
+
+		for (FString& StaticEndpoint : Settings->StaticEndpoints)
+		{
+			FIPv4Endpoint Endpoint;
+
+			if (ParseEndpoint(StaticEndpoint, Endpoint))
+			{
+				StaticEndpoints.Add(Endpoint);
+			}
+			else
+			{
+				UE_LOG(LogQuicMessaging, Warning,
+					TEXT("Invalid QUIC Messaging Static Endpoint '%s'"), *StaticEndpoint);
+			}
+		}
+
+        if (bResaveSettings)
+        {
+            Settings->SaveConfig();
+        }
+    	
+		UE_LOG(LogQuicMessaging, Log,
+			TEXT("Initialized bridge on interface %s"), *UnicastEndpoint.ToString());
+
+		const TSharedRef<FQuicEndpointConfig> EndpointConfig = (Settings->bIsClient)
+			? GetClientConfig(Settings) : GetServerConfig(Settings);
+
+		EndpointConfig->Endpoint = UnicastEndpoint;
+		EndpointConfig->LocalNodeId = Settings->EndpointGuid;
+		EndpointConfig->DiscoveryTimeoutSec = Settings->DiscoveryTimeoutSeconds;
+
+		EndpointConfig->EncryptionMode = (Settings->bEncryption)
+			? EEncryptionMode::Enabled : EEncryptionMode::Disabled;
+
+        const TSharedRef<FQuicMessageTransport, ESPMode::ThreadSafe> Transport =
+			MakeShared<FQuicMessageTransport, ESPMode::ThreadSafe>(
+				Settings->bIsClient, EndpointConfig, MoveTemp(StaticEndpoints));
+
+        Transport->OnTransportError().BindRaw(
+			this, &FQuicMessagingModule::HandleTransportError);
+
+        WeakBridgeTransport = Transport;
+        MessageBridge = FMessageBridgeBuilder().UsingTransport(Transport);
+    }
+
+    /**
+	 * Parse command line arguments to override QuicMessagingSettings.
+	 */
+	void ParseCommandLine(UQuicMessagingSettings* Settings, const TCHAR* CommandLine) const
+	{
+		if (Settings && CommandLine)
+		{
+			// Parse value overrides (if present)
+			FParse::Bool(CommandLine, 
+				TEXT("-QUICMESSAGING_TRANSPORT_ENABLE="),
+				Settings->EnableTransport);
+
+			FParse::Value(CommandLine,
+				TEXT("-QUICMESSAGING_TRANSPORT_UNICAST="),
+				Settings->UnicastEndpoint);
+
+			FParse::Bool(CommandLine,
+				TEXT("-QUICMESSAGING_TRANSPORT_ISCLIENT="),
+				Settings->bIsClient);
+
+			FParse::Bool(CommandLine,
+				TEXT("-QUICMESSAGING_TRANSPORT_ENCRYPTION="),
+				Settings->bEncryption);
+
+			FParse::Value(CommandLine,
+				TEXT("-QUICMESSAGING_DISCOVERY_TIMEOUT_SEC="),
+				Settings->DiscoveryTimeoutSeconds);
+
+			FParse::Bool(CommandLine,
+				TEXT("-QUICMESSAGING_AUTH_SERVER_ENABLED="),
+				Settings->bAuthEnabled);
+
+			FParse::Value(CommandLine,
+				TEXT("-QUICMESSAGING_MAX_AUTH_MESSAGE_SIZE_BYTES="),
+				Settings->MaxAuthenticationMessageSize);
+
+			FParse::Bool(CommandLine,
+				TEXT("-QUICMESSAGING_VERIFY_CERTIFICATE="),
+				Settings->bClientVerificationEnabled);
+
+			FParse::Value(CommandLine,
+				TEXT("-QUICMESSAGING_SERVER_CERTIFICATE="),
+				Settings->QuicServerCertificate);
+
+			FParse::Value(CommandLine,
+				TEXT("-QUICMESSAGING_SERVER_PRIVATEKEY="),
+				Settings->QuicServerPrivateKey);
+
+			FParse::Bool(CommandLine,
+				TEXT("-QUICMESSAGING_CONN_COOLDOWN_ENABLED="),
+				Settings->bConnectionCooldownEnabled);
+
+			FParse::Value(CommandLine, 
+				TEXT("-QUICMESSAGING_CONN_COOLDOWN_MAX_ATTEMPTS="),
+				Settings->ConnectionCooldownMaxAttempts);
+
+			FParse::Value(CommandLine,
+				TEXT("-QUICMESSAGING_CONN_COOLDOWN_PERIOD_SEC="),
+				Settings->ConnectionCooldownPeriodSeconds);
+
+			FParse::Value(CommandLine,
+				TEXT("-QUICMESSAGING_CONN_COOLDOWN_SEC="),
+				Settings->ConnectionCooldownSeconds);
+
+			FParse::Value(CommandLine,
+				TEXT("-QUICMESSAGING_CONN_COOLDOWN_MAX_SEC="),
+				Settings->ConnectionCooldownMaxSeconds);
+
+			FString StaticEndpoints;
+			FParse::Value(CommandLine,
+				TEXT("-QUICMESSAGING_TRANSPORT_STATIC="), StaticEndpoints, false);
+
+			TArray<FString> CommandLineStaticEndpoints;
+			StaticEndpoints.ParseIntoArray(CommandLineStaticEndpoints, TEXT(","));
+
+			for (const FString& CmdStaticEndpoint : CommandLineStaticEndpoints)
+			{
+				Settings->StaticEndpoints.AddUnique(CmdStaticEndpoint);
+			}
+		}
+	}
+
+	/** Get the config for client mode. */
+	TSharedRef<FQuicEndpointConfig> GetClientConfig(const UQuicMessagingSettings* InSettings)
+	{
+		FQuicClientConfig ClientEndpointConfig;
+
+		if (!InSettings)
+		{
+			UE_LOG(LogQuicMessaging, Fatal,
+				TEXT("[QuicMessagingModule::GetClientConfig] Pointer to settings not valid."));
+		}
+
+		ClientEndpointConfig.ClientVerificationMode = (InSettings->bClientVerificationEnabled)
+			? EQuicClientVerificationMode::Verify : EQuicClientVerificationMode::Pass;
+
+
+		return MakeShared<FQuicClientConfig>(ClientEndpointConfig);
+	}
+
+	/** Get the config for server mode. */
+	TSharedRef<FQuicEndpointConfig> GetServerConfig(const UQuicMessagingSettings* InSettings)
+	{
+		FQuicServerConfig ServerEndpointConfig;
+
+		if (!InSettings)
+		{
+			UE_LOG(LogQuicMessaging, Fatal,
+				TEXT("[QuicMessagingModule::GetServerConfig] Pointer to settings not valid."));
+		}
+
+		ServerEndpointConfig.Certificate = InSettings->QuicServerCertificate;
+		ServerEndpointConfig.PrivateKey = InSettings->QuicServerPrivateKey;
+
+		ServerEndpointConfig.MaxAuthenticationMessageSize
+			= InSettings->MaxAuthenticationMessageSize;
+
+		ServerEndpointConfig.AuthenticationMode = (InSettings->bAuthEnabled)
+			? EAuthenticationMode::Enabled : EAuthenticationMode::Disabled;
+
+		ServerEndpointConfig.ConnCooldownMode = (InSettings->bConnectionCooldownEnabled)
+			? EConnectionCooldownMode::Enabled : EConnectionCooldownMode::Disabled;
+
+		ServerEndpointConfig.ConnCooldownMaxAttempts
+			= InSettings->ConnectionCooldownMaxAttempts;
+
+		ServerEndpointConfig.ConnCooldownPeriodSec
+			= InSettings->ConnectionCooldownPeriodSeconds;
+
+		ServerEndpointConfig.ConnCooldownSec
+			= InSettings->ConnectionCooldownSeconds;
+
+		ServerEndpointConfig.ConnCooldownMaxSec
+			= InSettings->ConnectionCooldownMaxSeconds;
+
+		return MakeShared<FQuicServerConfig>(ServerEndpointConfig);
+	}
+
+    /** Shuts down the message bridge. */
+    void ShutdownBridge()
+    {
+        StopAutoRepairRoutine();
+
+        if (WeakBridgeTransport.IsValid())
+        {
+			WeakBridgeTransport.Reset();
+        }
+
+        if (MessageBridge.IsValid())
+        {
+            MessageBridge->Disable();
+        	
+            FPlatformProcess::Sleep(0.2f);
+
+			if (MessageBridge.IsValid())
+			{
+				MessageBridge.Reset();
+			}
+        }
+    }
+
+private:
+
+	void CreateEndpointGuid() const
+	{
+		UQuicMessagingSettings* Settings = GetMutableDefault<UQuicMessagingSettings>();
+		Settings->EndpointGuid = FGuid::NewGuid();
+		Settings->SaveConfig();
+	}
+
+    /** Callback for when the application has been reactivated. */
+    void HandleApplicationHasReactivated()
+	{
+		RestartServices();
+	}
+
+	/** Callback for when the application will be deactivated.*/
+	void HandleApplicationWillDeactivate()
+	{
+		ShutdownServices();
+	}
+
+    /** Callback for when the settings were saved. */
+    bool HandleSettingsSaved()
+    {
+		UE_LOG(LogQuicMessaging, Display,
+			TEXT("QUIC settings saved, restarting services..."));
+
+        RestartServices();
+
+        return true;
+    }
+
+    void HandleAppPreExit() const
+    {
+        if (TSharedPtr<FQuicMessageTransport> QuicTransport = WeakBridgeTransport.Pin())
+        {
+            QuicTransport->OnAppPreExit();
+        }
+    }
+
+    void HandleTransportError()
+    {
+        const UQuicMessagingSettings* Settings = GetDefault<UQuicMessagingSettings>();
+        if (Settings->bAutoRepair)
+        {
+            StartAutoRepairRoutine(Settings->AutoRepairAttemptLimit);
+        }
+        else
+        {
+			UE_LOG(LogQuicMessaging, Error,
+				TEXT("QUIC messaging encountered an error. "
+				"Please restart the service for proper functionality"));
+        }
+    }
+
+    void StartAutoRepairRoutine(uint32 MaxRetryAttempt)
+    {
+        StopAutoRepairRoutine();
+
+        FTimespan CheckDelay(0, 0, 1);
+        uint32 CheckNumber = 1;
+
+        AutoRepairHandle = FTSTicker::GetCoreTicker().AddTicker(
+			FTickerDelegate::CreateLambda(
+				[WeakTransport = WeakBridgeTransport, LastTime = FDateTime::UtcNow(),
+				CheckDelay, CheckNumber, MaxRetryAttempt](float DeltaTime) mutable
+        {
+            QUICK_SCOPE_CYCLE_COUNTER(STAT_FQuicMessagingModule_AutoRepair);
+
+			bool bContinue = true;
+            const FDateTime UtcNow = FDateTime::UtcNow();
+
+			if (LastTime + (CheckDelay * CheckNumber) <= UtcNow)
+			{
+				if (const TSharedPtr<FQuicMessageTransport> Transport = WeakTransport.Pin())
+				{
+					// if the restart fail, continue the routine if we are still under the retry attempt limit
+					bContinue = !Transport->RestartTransport();
+					bContinue = bContinue && CheckNumber <= MaxRetryAttempt;
+				}
+				// if we do not have a valid transport also stop the routine
+				else
+				{
+					bContinue = false;
+				}
+				++CheckNumber;
+				LastTime = UtcNow;
+			}
+			return bContinue;
+			
+        }), 1.0f);
+
+		UE_LOG(LogQuicMessaging, Warning,
+			TEXT("QUIC messaging encountered an error. "
+			"Auto repair routine started for reinitialization"));
+    }
+
+    void StopAutoRepairRoutine() const
+    {
+        if (AutoRepairHandle.IsValid())
+        {
+			FTSTicker::GetCoreTicker().RemoveTicker(AutoRepairHandle);
+        }
+    }
+
+private:
+
+    /** Name of the network messaging extension. */
+    static FName QuicMessagingName;
+
+    /** Holds the message bridge if present. */
+	TSharedPtr<IMessageBridge, ESPMode::ThreadSafe> MessageBridge = nullptr;
+
+	/** Holds the bridge transport if present.  */
+	TWeakPtr<FQuicMessageTransport, ESPMode::ThreadSafe> WeakBridgeTransport = nullptr;
+
+	/** Critical section protecting access to the transport static endpoints and additional static endpoints. */
+	FCriticalSection StaticEndpointsCS;
+
+	/** Holds additional static endpoints added through the modular feature interface. */
+	TSet<FIPv4Endpoint> AdditionalStaticEndpoints = TSet<FIPv4Endpoint>();
+
+	/** Holds the delegate handle for the auto repair routine. */
+	FTSTicker::FDelegateHandle AutoRepairHandle = nullptr;
+};
+
+FName FQuicMessagingModule::QuicMessagingName("QuicMessaging");
+
+
+void EmptyLinkFunctionForStaticInitializationQuicMessagingTests()
+{
+	// Force references to the object files containing the functions below, to prevent them being
+	// excluded by the linker when the plug-in is compiled into a static library for monolithic builds.
+
+	extern void EmptyLinkFunctionForStaticInitializationQuicMessageTransportTest();
+	EmptyLinkFunctionForStaticInitializationQuicMessageTransportTest();
+}
+
+
+IMPLEMENT_MODULE(FQuicMessagingModule, QuicMessaging);
+
+
+#undef LOCTEXT_NAMESPACE
diff --git a/Engine/Plugins/Experimental/QuicMessaging/Source/QuicMessaging/Private/QuicMessagingPrivate.h b/Engine/Plugins/Experimental/QuicMessaging/Source/QuicMessaging/Private/QuicMessagingPrivate.h
new file mode 100644
--- /dev/null
+++ b/Engine/Plugins/Experimental/QuicMessaging/Source/QuicMessaging/Private/QuicMessagingPrivate.h
@@ -0,0 +1,27 @@
+// Copyright Epic Games, Inc. All Rights Reserved.
+
+#pragma once
+
+#include "CoreMinimal.h"
+#include "Logging/LogMacros.h"
+#include "Interfaces/IPv4/IPv4Address.h"
+#include "Interfaces/IPv4/IPv4Endpoint.h"
+
+
+/** Declares a log category for this module. */
+DECLARE_LOG_CATEGORY_EXTERN(LogQuicMessaging, VeryVerbose, All);
+
+
+/** Defines the maximum number of annotations a message can have. */
+#define QUIC_MESSAGING_MAX_ANNOTATIONS 128
+
+/** Defines the maximum number of recipients a message can have. */
+#define QUIC_MESSAGING_MAX_RECIPIENTS 1024
+
+
+/** Defines the MsQuic version to be used. */
+static FString MSQUIC_VERSION = "v220";
+
+/** Defines the MsQuic binaries path. */
+static FString MSQUIC_BINARIES_PATH = "Binaries/ThirdParty/MsQuic/" + MSQUIC_VERSION;
+
diff --git a/Engine/Plugins/Experimental/QuicMessaging/Source/QuicMessaging/Private/QuicMessagingSettings.cpp b/Engine/Plugins/Experimental/QuicMessaging/Source/QuicMessaging/Private/QuicMessagingSettings.cpp
new file mode 100644
--- /dev/null
+++ b/Engine/Plugins/Experimental/QuicMessaging/Source/QuicMessaging/Private/QuicMessagingSettings.cpp
@@ -0,0 +1,9 @@
+// Copyright Epic Games, Inc. All Rights Reserved.
+
+#include "Shared/QuicMessagingSettings.h"
+
+
+UQuicMessagingSettings::UQuicMessagingSettings(const FObjectInitializer& ObjectInitializer)
+    : Super(ObjectInitializer)
+	, EnableTransport(true)
+{ }
\ No newline at end of file
diff --git a/Engine/Plugins/Experimental/QuicMessaging/Source/QuicMessaging/Private/Serialization/QuicDeserializedMessage.h b/Engine/Plugins/Experimental/QuicMessaging/Source/QuicMessaging/Private/Serialization/QuicDeserializedMessage.h
new file mode 100644
--- /dev/null
+++ b/Engine/Plugins/Experimental/QuicMessaging/Source/QuicMessaging/Private/Serialization/QuicDeserializedMessage.h
@@ -0,0 +1,226 @@
+// Copyright Epic Games, Inc. All Rights Reserved.
+
+#pragma once
+
+#include "CoreTypes.h"
+#include "Containers/Array.h"
+#include "IMessageContext.h"
+#include "Misc/DateTime.h"
+#include "Templates/SharedPointer.h"
+#include "UObject/WeakObjectPtrTemplates.h"
+
+#include "QuicMessages.h"
+
+
+class IMessageAttachment;
+class UScriptStruct;
+
+
+/**
+ * Enumerates possibly states of a deserialized message.
+ */
+enum class EQuicDeserializedMessageState
+{
+    /** The message data is complete. */
+	Complete,
+
+	/** The message data is incomplete. */
+	Incomplete,
+
+	/** The message data is invalid. */
+	Invalid
+};
+
+
+/**
+ * Holds deserialized message data.
+ */
+class FQuicDeserializedMessage
+	: public IMessageContext
+{
+
+	friend class FQuicDeserializeMessageTask;
+
+public:
+
+	/** Default constructor. */
+	FQuicDeserializedMessage() = delete;
+
+	/**
+	 * Create and initialize a new instance.
+	 */
+	FQuicDeserializedMessage(FInboundMessage InInboundMessage)
+		: State(EQuicDeserializedMessageState::Incomplete)
+		, InboundMessage(InInboundMessage)
+		, Annotations(TMap<FName, FString>())
+		, Attachment(nullptr)
+		, Expiration(FDateTime::UtcNow())
+		, MessageData(nullptr)
+		, Recipients(TArray<FMessageAddress>())
+		, Scope(EMessageScope::All)
+		, Flags(EMessageFlags::None)
+		, Sender(FMessageAddress())
+		, TimeSent(FDateTime::UtcNow())
+		, TypeInfo(nullptr)
+	{ }
+
+	/** Virtual destructor. */
+	virtual ~FQuicDeserializedMessage() override
+	{
+		if (MessageData != nullptr)
+		{
+			if (TypeInfo.IsValid())
+			{
+				TypeInfo->DestroyStruct(MessageData);
+			}
+
+			FMemory::Free(MessageData);
+			MessageData = nullptr;
+		}
+	}
+
+public:
+
+	/**
+	 * Gets the state of the message data.
+	 */
+	EQuicDeserializedMessageState GetState() const
+	{
+		return State;
+	}
+
+	/**
+	 * Updates the state of this message data.
+	 */
+	void UpdateState(EQuicDeserializedMessageState InState)
+	{
+		State = InState;
+	}
+
+	/**
+	 * Get the inbound message.
+	 */
+	FInboundMessage& GetInboundMessage()
+	{
+		return InboundMessage;
+	}
+
+public:
+
+	//~ IMessageContext interface
+
+	virtual const TMap<FName, FString>& GetAnnotations() const override
+	{
+		return Annotations;
+	}
+
+	virtual TSharedPtr<IMessageAttachment, ESPMode::ThreadSafe> GetAttachment() const override
+	{
+		return Attachment;
+	}
+
+	virtual const FDateTime& GetExpiration() const override
+	{
+		return Expiration;
+	}
+
+	virtual const void* GetMessage() const override
+	{
+		return MessageData;
+	}
+
+	virtual const TWeakObjectPtr<UScriptStruct>& GetMessageTypeInfo() const override
+	{
+		return TypeInfo;
+	}
+
+	virtual TSharedPtr<IMessageContext, ESPMode::ThreadSafe> GetOriginalContext() const override
+	{
+		return nullptr;
+	}
+
+	virtual const TArray<FMessageAddress>& GetRecipients() const override
+	{
+		return Recipients;
+	}
+
+	virtual EMessageScope GetScope() const override
+	{
+		return Scope;
+	}
+
+	virtual EMessageFlags GetFlags() const override
+	{
+		return Flags;
+	}
+
+	virtual const FMessageAddress& GetSender() const override
+	{
+		return Sender;
+	}
+
+	virtual const FMessageAddress& GetForwarder() const override
+	{
+		return Sender;
+	}
+
+	virtual ENamedThreads::Type GetSenderThread() const override
+	{
+		return ENamedThreads::AnyThread;
+	}
+
+	virtual const FDateTime& GetTimeForwarded() const override
+	{
+		return TimeSent;
+	}
+
+	virtual const FDateTime& GetTimeSent() const override
+	{
+		return TimeSent;
+	}
+
+private:
+
+	/** Holds the deserialized message state. */
+	EQuicDeserializedMessageState State;
+
+	/** Holds the inbound message. */
+	FInboundMessage InboundMessage;
+
+private:
+
+	/** Holds the optional message annotations. */
+	TMap<FName, FString> Annotations;
+
+	/** Holds a pointer to attached binary data. */
+	TSharedPtr<IMessageAttachment, ESPMode::ThreadSafe> Attachment;
+
+	/** Holds the expiration time. */
+	FDateTime Expiration;
+
+	/** Holds the message. */
+	void* MessageData;
+
+	/** Holds the message recipients. */
+	TArray<FMessageAddress> Recipients;
+
+	/** Holds the message's scope. */
+	EMessageScope Scope;
+
+	/** Holds the message's flags. */
+	EMessageFlags Flags;
+
+	/** Holds the sender's identifier. */
+	FMessageAddress Sender;
+
+	/** Holds the time at which the message was sent. */
+	FDateTime TimeSent;
+
+	/** Holds the message's type information. */
+	TWeakObjectPtr<UScriptStruct> TypeInfo;
+
+};
+
+
+
+
diff --git a/Engine/Plugins/Experimental/QuicMessaging/Source/QuicMessaging/Private/Serialization/QuicSerializationTasks.cpp b/Engine/Plugins/Experimental/QuicMessaging/Source/QuicMessaging/Private/Serialization/QuicSerializationTasks.cpp
new file mode 100644
--- /dev/null
+++ b/Engine/Plugins/Experimental/QuicMessaging/Source/QuicMessaging/Private/Serialization/QuicSerializationTasks.cpp
@@ -0,0 +1,346 @@
+// Copyright Epic Games, Inc. All Rights Reserved.
+
+#include "QuicSerializationTasks.h"
+#include "QuicMessagingPrivate.h"
+
+#include "HAL/Event.h"
+#include "IMessageContext.h"
+#include "Backends/JsonStructSerializerBackend.h"
+#include "Backends/CborStructSerializerBackend.h"
+#include "StructSerializer.h"
+#include "Backends/JsonStructDeserializerBackend.h"
+#include "Backends/CborStructDeserializerBackend.h"
+#include "StructDeserializer.h"
+
+#include "QuicMessagingSettings.h"
+#include "QuicSerializedMessage.h"
+#include "QuicDeserializedMessage.h"
+
+
+/**
+ * FQuicSerializeMessageTask
+ */
+
+void FQuicSerializeMessageTask::SerializeMessage(FArchive& Archive, EQuicMessageFormat MessageFormat,
+	const TSharedRef<IMessageContext, ESPMode::ThreadSafe>& InMessageContext,
+	const EStructSerializerBackendFlags StructSerializerBackendFlags)
+{
+	const FName& MessageType = InMessageContext->GetMessageTypePathName().GetAssetName();
+	Archive << const_cast<FName&>(MessageType);
+
+	const FMessageAddress& Sender = InMessageContext->GetSender();
+	Archive << const_cast<FMessageAddress&>(Sender);
+
+	const TArray<FMessageAddress>& Recipients = InMessageContext->GetRecipients();
+	Archive << const_cast<TArray<FMessageAddress>&>(Recipients);
+
+	EMessageScope Scope = InMessageContext->GetScope();
+	Archive << Scope;
+
+	EMessageFlags Flags = InMessageContext->GetFlags();
+	Archive << Flags;
+
+	const FDateTime& TimeSent = InMessageContext->GetTimeSent();
+	Archive << const_cast<FDateTime&>(TimeSent);
+
+	const FDateTime& Expiration = InMessageContext->GetExpiration();
+	Archive << const_cast<FDateTime&>(Expiration);
+
+	int32 NumAnnotations = InMessageContext->GetAnnotations().Num();
+	Archive << NumAnnotations;
+
+	for (const auto& AnnotationPair : InMessageContext->GetAnnotations())
+	{
+		Archive << const_cast<FName&>(AnnotationPair.Key);
+		Archive << const_cast<FString&>(AnnotationPair.Value);
+	}
+
+	// Message Wire Format Id
+	check(MessageFormat == EQuicMessageFormat::CborPlatformEndianness
+		|| MessageFormat == EQuicMessageFormat::CborStandardEndianness);
+	
+	uint8 Format = (uint8)(MessageFormat);
+	Archive << Format;
+
+	// Serialize message body with cbor
+	FCborStructSerializerBackend Backend(Archive,
+		StructSerializerBackendFlags | (MessageFormat == EQuicMessageFormat::CborStandardEndianness
+		? EStructSerializerBackendFlags::WriteCborStandardEndianness : EStructSerializerBackendFlags::None));
+	
+	FStructSerializer::Serialize(InMessageContext->GetMessage(),
+        *InMessageContext->GetMessageTypeInfo(), Backend);
+}
+
+
+void FQuicSerializeMessageTask::DoTask(ENamedThreads::Type CurrentThread,
+    const FGraphEventRef& MyCompletionGraphEvent)
+{
+    if (MessageContext->IsValid())
+    {
+		FMemoryWriter MemoryWriter(*SerializedMessage->GetDataPtr(), true);
+
+        FArchive& Archive = MemoryWriter;
+
+        SerializeMessage(Archive, SerializedMessage->GetFormat(), MessageContext,
+            EStructSerializerBackendFlags::Default);
+        
+        SerializedMessage->UpdateState(EQuicSerializedMessageState::Complete);
+    }
+    else
+    {
+        SerializedMessage->UpdateState(EQuicSerializedMessageState::Invalid);
+    }
+
+    // Signal task completion
+	TSharedPtr<FEvent, ESPMode::ThreadSafe> CompletionEvent = CompletionEventPtr.Pin();
+
+	if (CompletionEvent.IsValid())
+	{
+		CompletionEvent->Trigger();
+	}
+}
+
+
+ENamedThreads::Type FQuicSerializeMessageTask::GetDesiredThread()
+{
+	return ENamedThreads::AnyThread;
+}
+
+
+TStatId FQuicSerializeMessageTask::GetStatId() const
+{
+	RETURN_QUICK_DECLARE_CYCLE_STAT(FQuicSerializeMessageTask, STATGROUP_TaskGraphTasks);
+}
+
+
+ESubsequentsMode::Type FQuicSerializeMessageTask::GetSubsequentsMode()
+{
+	return ESubsequentsMode::FireAndForget;
+}
+
+
+/**
+ * FQuicDeserializeMessageTask
+ */
+
+bool FQuicDeserializeMessageTask::DeserializeMessage(FMemoryReader& MessageReader)
+{
+
+    // message type info
+	{
+		FName MessageType;
+		MessageReader << MessageType;
+
+		// @todo gmp: cache message types for faster lookup
+		DeserializedMessage->TypeInfo = FindFirstObjectSafe<UScriptStruct>(*MessageType.ToString(), EFindFirstObjectOptions::EnsureIfAmbiguous);
+
+		if (!DeserializedMessage->TypeInfo.IsValid(false, true))
+		{
+			return false;
+		}
+	}
+
+	// Sender address
+	{
+		MessageReader << DeserializedMessage->Sender;
+	}
+
+	// Recipient addresses
+	{
+		int32 NumRecipients = 0;
+		MessageReader << NumRecipients;
+
+		if ((NumRecipients < 0) || (NumRecipients > QUIC_MESSAGING_MAX_RECIPIENTS))
+		{
+			return false;
+		}
+
+		DeserializedMessage->Recipients.Empty(NumRecipients);
+
+		while (0 < NumRecipients--)
+		{
+			MessageReader << DeserializedMessage->Recipients.AddDefaulted_GetRef();
+		}
+	}
+
+	// Message scope
+	{
+		MessageReader << DeserializedMessage->Scope;
+
+		if (DeserializedMessage->Scope > EMessageScope::All)
+		{
+			return false;
+		}
+	}
+
+	// Message flags
+	{
+		MessageReader << DeserializedMessage->Flags;
+	}
+
+	// Time sent & expiration
+	{
+		MessageReader << DeserializedMessage->TimeSent;
+		MessageReader << DeserializedMessage->Expiration;
+	}
+
+	// Annotations
+	{
+		int32 NumAnnotations = 0;
+		MessageReader << NumAnnotations;
+
+		if (NumAnnotations > QUIC_MESSAGING_MAX_ANNOTATIONS)
+		{
+			return false;
+		}
+
+		while (0 < NumAnnotations--)
+		{
+			FName Key;
+			FString Value;
+
+			MessageReader << Key;
+			MessageReader << Value;
+
+			DeserializedMessage->Annotations.Add(Key, Value);
+		}
+	}
+
+	// Wire format 
+	uint8 FormatId;
+	MessageReader << FormatId;
+	EQuicMessageFormat MessageFormat = (EQuicMessageFormat)FormatId;
+
+	// Create message body
+	DeserializedMessage->MessageData = FMemory::Malloc(DeserializedMessage->TypeInfo->GetStructureSize());
+	DeserializedMessage->TypeInfo->InitializeStruct(DeserializedMessage->MessageData);
+
+	switch (MessageFormat)
+	{
+
+	case EQuicMessageFormat::Json:
+	{
+		// Deserialize json
+		FJsonStructDeserializerBackend Backend(MessageReader);
+		bool DeserializeResult = FStructDeserializer::Deserialize(DeserializedMessage->MessageData, *DeserializedMessage->TypeInfo, Backend);
+
+		if (!DeserializeResult)
+		{
+			UE_LOG(LogQuicMessaging, Error, TEXT("[DeserializeMessage] Json deserialize result failed"));
+		}
+
+		return DeserializeResult;
+	}
+	break;
+
+	case EQuicMessageFormat::CborPlatformEndianness:
+	{
+		// Deserialize cbor (using this platform endianness)
+		FCborStructDeserializerBackend Backend(MessageReader, ECborEndianness::Platform);
+		bool DeserializeResult = FStructDeserializer::Deserialize(DeserializedMessage->MessageData, *DeserializedMessage->TypeInfo, Backend);
+
+		if (!DeserializeResult)
+		{
+			UE_LOG(LogQuicMessaging, Error, TEXT("[DeserializeMessage] CborPlatform deserialize result failed"));
+		}
+
+		return DeserializeResult;
+	}
+	break;
+
+	case EQuicMessageFormat::CborStandardEndianness:
+	{
+		// Deserialize cbor (using the CBOR standard endianness - big endian)
+		FCborStructDeserializerBackend Backend(MessageReader, ECborEndianness::StandardCompliant);
+		bool DeserializeResult = FStructDeserializer::Deserialize(DeserializedMessage->MessageData, *DeserializedMessage->TypeInfo, Backend);
+
+		if (!DeserializeResult)
+		{
+			UE_LOG(LogQuicMessaging, Error, TEXT("[DeserializeMessage] CborStandard deserialize result failed"));
+		}
+
+		return DeserializeResult;
+	}
+	break;
+
+	case EQuicMessageFormat::TaggedProperty:
+	{
+		// Deserialize message body using tagged property
+		// Hack: this binary serialization should use a more standard protocol, should use cbor
+		DeserializedMessage->TypeInfo->SerializeItem(MessageReader, DeserializedMessage->MessageData, nullptr);
+		bool DeserializeResult = !MessageReader.GetError();
+
+		if (!DeserializeResult)
+		{
+			UE_LOG(LogQuicMessaging, Error, TEXT("[DeserializeMessage] TaggedProperty deserialize result failed"));
+		}
+
+		return DeserializeResult;
+	}
+	break;
+
+	default:
+
+		// Unsupported format
+		return false;
+		
+	}
+}
+
+
+void FQuicDeserializeMessageTask::DoTask(ENamedThreads::Type CurrentThread,
+    const FGraphEventRef& MyCompletionGraphEvent)
+{
+    if (Data.IsValid())
+    {
+		FMemoryReader MessageReader(*Data.Get(), true);
+		MessageReader.ArMaxSerializeSize = NAME_SIZE;
+
+		bool DeserializeMessageResult = DeserializeMessage(MessageReader);
+
+		if (DeserializeMessageResult)
+		{
+			DeserializedMessage->UpdateState(EQuicDeserializedMessageState::Complete);
+		}
+		else
+		{
+			DeserializedMessage->UpdateState(EQuicDeserializedMessageState::Invalid);
+		}
+
+		Data.Reset();
+    }
+    else
+    {
+        DeserializedMessage->UpdateState(EQuicDeserializedMessageState::Invalid);
+    }
+
+    // signal task completion
+	TSharedPtr<FEvent, ESPMode::ThreadSafe> CompletionEvent = CompletionEventPtr.Pin();
+
+	if (CompletionEvent.IsValid())
+	{
+		CompletionEvent->Trigger();
+	}
+}
+
+
+ENamedThreads::Type FQuicDeserializeMessageTask::GetDesiredThread()
+{
+	return ENamedThreads::AnyThread;
+}
+
+
+TStatId FQuicDeserializeMessageTask::GetStatId() const
+{
+	RETURN_QUICK_DECLARE_CYCLE_STAT(FQuicDeserializeMessageTask, STATGROUP_TaskGraphTasks);
+}
+
+
+ESubsequentsMode::Type FQuicDeserializeMessageTask::GetSubsequentsMode()
+{
+	return ESubsequentsMode::FireAndForget;
+}
+
+
+
+
diff --git a/Engine/Plugins/Experimental/QuicMessaging/Source/QuicMessaging/Private/Serialization/QuicSerializationTasks.h b/Engine/Plugins/Experimental/QuicMessaging/Source/QuicMessaging/Private/Serialization/QuicSerializationTasks.h
new file mode 100644
--- /dev/null
+++ b/Engine/Plugins/Experimental/QuicMessaging/Source/QuicMessaging/Private/Serialization/QuicSerializationTasks.h
@@ -0,0 +1,177 @@
+// Copyright Epic Games, Inc. All Rights Reserved.
+
+#pragma once
+
+#include "Async/TaskGraphInterfaces.h"
+#include "Stats/Stats.h"
+#include "Templates/SharedPointer.h"
+
+
+class FEvent;
+class FMemoryReader;
+class FQuicSerializedMessage;
+class FQuicDeserializedMessage;
+class IMessageContext;
+enum class EStructSerializerBackendFlags;
+enum class EQuicMessageFormat : uint8;
+
+
+class FQuicSerializeMessageTask
+{
+
+public:
+
+	FQuicSerializeMessageTask() = delete;
+
+	/**
+	 * Creates a new serialize message task.
+	 *
+	 * @param InMessageContext The message context
+	 * @param InSerializedMessage The serialized message
+	 * @param InCompletionEvent The completion event
+	 */
+	FQuicSerializeMessageTask(
+		TSharedRef<IMessageContext, ESPMode::ThreadSafe> InMessageContext,
+		TSharedPtr<FQuicSerializedMessage, ESPMode::ThreadSafe> InSerializedMessage,
+		TSharedPtr<FEvent, ESPMode::ThreadSafe> InCompletionEvent)
+		: MessageContext(InMessageContext)
+		, SerializedMessage(InSerializedMessage)
+		, CompletionEventPtr(InCompletionEvent)
+	{ }
+
+private:
+
+	/**
+	 * Serializes the message.
+	 *
+	 * @param Archive The archive where data will be stored
+	 * @param MessageFormat The message format
+	 * @param InMessageContext Shared ref of the message context
+	 * @param StructSerializerBackendFlags Flags for the struct serializer backend
+	 */
+	void SerializeMessage(FArchive& Archive, EQuicMessageFormat MessageFormat,
+		const TSharedRef<IMessageContext, ESPMode::ThreadSafe>& InMessageContext,
+		const EStructSerializerBackendFlags StructSerializerBackendFlags);
+
+public:
+
+	/**
+	 * Performs the actual task.
+	 *
+	 * @param CurrentThread The thread that this task is executing on
+	 * @param MyCompletionGraphEvent The completion event
+	 */
+	void DoTask(ENamedThreads::Type CurrentThread,
+		const FGraphEventRef& MyCompletionGraphEvent);
+
+	/**
+	 * Returns the name of the thread that this task should run on.
+	 *
+	 * @return The thread name
+	 */
+	ENamedThreads::Type GetDesiredThread();
+
+	/**
+	 * Gets the task's stats tracking identifier.
+	 *
+	 * @return Stats identifier
+	 */
+	TStatId GetStatId() const;
+
+	/**
+	 * Gets the mode for tracking subsequent tasks.
+	 *
+	 * @return Always track subsequent tasks
+	 */
+	static ESubsequentsMode::Type GetSubsequentsMode();
+
+private:
+
+	/** The context of the message to serialize. */
+	TSharedRef<IMessageContext, ESPMode::ThreadSafe> MessageContext;
+
+	/** A reference to the serialized message data. */
+	TSharedPtr<FQuicSerializedMessage, ESPMode::ThreadSafe> SerializedMessage;
+
+	/** An event signaling that the message was serialized. */
+	TWeakPtr<FEvent, ESPMode::ThreadSafe> CompletionEventPtr;
+
+};
+
+
+class FQuicDeserializeMessageTask
+{
+
+public:
+
+	FQuicDeserializeMessageTask() = delete;
+
+	/**
+	 * Creates a new deserialize message task.
+	 *
+	 * @param InData The data
+	 * @param InDeserializedMessage The deserialized message
+	 * @param InCompletionEvent The completion event
+	 */
+	FQuicDeserializeMessageTask(
+		TSharedPtr<TArray<uint8>, ESPMode::ThreadSafe> InData,
+		TSharedPtr<FQuicDeserializedMessage, ESPMode::ThreadSafe> InDeserializedMessage,
+		TSharedPtr<FEvent, ESPMode::ThreadSafe> InCompletionEvent)
+		: Data(InData)
+		, DeserializedMessage(InDeserializedMessage)
+		, CompletionEventPtr(InCompletionEvent)
+	{ }
+
+private:
+
+	/**
+	 * Deserializes the message.
+	 *
+	 * @param MessageReader The already initialized message reader
+	 */
+	bool DeserializeMessage(FMemoryReader& MessageReader);
+
+public:
+
+	/**
+	 * Performs the actual task.
+	 *
+	 * @param CurrentThread The thread that this task is executing on
+	 * @param MyCompletionGraphEvent The completion event
+	 */
+	void DoTask(ENamedThreads::Type CurrentThread,
+		const FGraphEventRef& MyCompletionGraphEvent);
+
+	/**
+	 * Returns the name of the thread that this task should run on.
+	 *
+	 * @return The thread name
+	 */
+	ENamedThreads::Type GetDesiredThread();
+
+	/**
+	 * Gets the task's stats tracking identifier.
+	 *
+	 * @return Stats identifier
+	 */
+	TStatId GetStatId() const;
+
+	/**
+	 * Gets the mode for tracking subsequent tasks.
+	 *
+	 * @return Always track subsequent tasks
+	 */
+	static ESubsequentsMode::Type GetSubsequentsMode();
+
+private:
+
+	/** A pointer to the unserialized message data. */
+	TSharedPtr<TArray<uint8>, ESPMode::ThreadSafe> Data;
+
+	/** A pointer to the deserialized message. */
+	TSharedPtr<FQuicDeserializedMessage, ESPMode::ThreadSafe> DeserializedMessage;
+
+	/** An event signaling that the message was serialized. */
+	TWeakPtr<FEvent, ESPMode::ThreadSafe> CompletionEventPtr;
+
+};
diff --git a/Engine/Plugins/Experimental/QuicMessaging/Source/QuicMessaging/Private/Serialization/QuicSerializedMessage.h b/Engine/Plugins/Experimental/QuicMessaging/Source/QuicMessaging/Private/Serialization/QuicSerializedMessage.h
new file mode 100644
--- /dev/null
+++ b/Engine/Plugins/Experimental/QuicMessaging/Source/QuicMessaging/Private/Serialization/QuicSerializedMessage.h
@@ -0,0 +1,151 @@
+// Copyright Epic Games, Inc. All Rights Reserved.
+
+#pragma once
+
+#include "CoreTypes.h"
+#include "Containers/Array.h"
+#include "Delegates/Delegate.h"
+#include "Templates/SharedPointer.h"
+#include "Serialization/Archive.h"
+#include "Serialization/MemoryWriter.h"
+#include "Serialization/MemoryReader.h"
+#include "IMessageContext.h"
+
+#include "QuicMessagingSettings.h"
+#include "QuicFlags.h"
+
+
+/**
+ * Enumerates possibly states of a serialized message.
+ */
+enum class EQuicSerializedMessageState
+{
+    /** The message data is complete. */
+	Complete,
+
+	/** The message data is incomplete. */
+	Incomplete,
+
+	/** The message data is invalid. */
+	Invalid
+};
+
+
+/**
+ * Holds serialized message data.
+ */
+class FQuicSerializedMessage
+{
+
+public:
+
+	FQuicSerializedMessage() = delete;
+
+    /** Default constructor */
+	FQuicSerializedMessage(EQuicMessageFormat InMessageFormat, EMessageFlags InFlags,
+		TArray<FGuid> InRecipients, EQuicMessageType InMessageType,
+		TSharedPtr<TArray<uint8>, ESPMode::ThreadSafe> InDataPtr,
+		TSharedRef<IMessageContext, ESPMode::ThreadSafe> InMessageContext)
+		: DataPtr(InDataPtr)
+		, MessageContext(InMessageContext)
+		, State(EQuicSerializedMessageState::Incomplete)
+		, Flags(InFlags)
+		, Format(InMessageFormat)
+		, Recipients(InRecipients)
+		, MessageType(InMessageType)
+	{ }
+
+public:
+
+	/**
+	 * Get the pointer to the data.
+	 */
+	TSharedPtr<TArray<uint8>, ESPMode::ThreadSafe> GetDataPtr()
+	{
+		return DataPtr;
+	}
+
+	/**
+	 * Get the message context.
+	 */
+	TSharedRef<IMessageContext, ESPMode::ThreadSafe> GetMessageContext()
+	{
+		return MessageContext;
+	}
+
+	/**
+	 * Gets the state of the message data.
+	 */
+	EQuicSerializedMessageState GetState() const
+	{
+		return State;
+	}
+
+	/**
+	 * Updates the state of this message data.
+	 */
+	void UpdateState(EQuicSerializedMessageState InState)
+	{
+		State = InState;
+	}
+
+	/**
+	 * Get the message flags.
+	 */
+	EMessageFlags GetFlags() const
+	{
+		return Flags;
+	}
+
+	/**
+	 * Get the format used to encode the message
+	 */
+	EQuicMessageFormat GetFormat() const
+	{
+		return Format;
+	}
+
+	/**
+	 * Get the message recipients.
+	 */
+	TArray<FGuid> GetRecipients() const
+	{
+		return Recipients;
+	}
+
+	/**
+	 * Get the message type.
+	 */
+	EQuicMessageType GetMessageType() const
+	{
+		return MessageType;
+	}
+
+private:
+
+	/** Holds the serialized data. */
+	TSharedPtr<TArray<uint8>, ESPMode::ThreadSafe> DataPtr;
+
+	/** Holds the message context. */
+	TSharedRef<IMessageContext, ESPMode::ThreadSafe> MessageContext;
+
+	/** Holds the message data state. */
+	EQuicSerializedMessageState State;
+
+	/** Holds message flags, captured from context. */
+	EMessageFlags Flags;
+
+	/** The format used to serialize the message. */
+	EQuicMessageFormat Format;
+
+	/** Holds the message recipients. */
+	TArray<FGuid> Recipients;
+
+	/** Holds the message type. */
+	EQuicMessageType MessageType;
+
+};
+
+
+
+
diff --git a/Engine/Plugins/Experimental/QuicMessaging/Source/QuicMessaging/Private/Tests/QuicMessageTransportTest.cpp b/Engine/Plugins/Experimental/QuicMessaging/Source/QuicMessaging/Private/Tests/QuicMessageTransportTest.cpp
new file mode 100644
--- /dev/null
+++ b/Engine/Plugins/Experimental/QuicMessaging/Source/QuicMessaging/Private/Tests/QuicMessageTransportTest.cpp
@@ -0,0 +1,677 @@
+// Copyright Epic Games, Inc. All Rights Reserved.
+
+#include "CoreMinimal.h"
+#include "Misc/Guid.h"
+#include "Misc/AutomationTest.h"
+#include "HAL/Runnable.h"
+#include "HAL/RunnableThread.h"
+#include "HAL/PlatformProcess.h"
+#include "IMessageContext.h"
+#include "IMessageTransportHandler.h"
+
+#include "Tests/QuicMessagingTestTypes.h"
+#include "Transport/QuicMessageTransport.h"
+#include "IQuicNetworkMessagingExtension.h"
+#include "QuicEndpointConfig.h"
+#include "QuicTransportMessages.h"
+#include "Shared/QuicMessagingSettings.h"
+#include "Features/IModularFeatures.h"
+#include "MessageEndpoint.h"
+
+
+IMPLEMENT_SIMPLE_AUTOMATION_TEST(FQuicMessageTransportTest,
+	"System.Core.Messaging.Transports.Quic.QuicMessageTransport (may take some minutes!)",
+	EAutomationTestFlags::ApplicationContextMask | EAutomationTestFlags::EngineFilter)
+
+
+/**
+ * QuicMessaging test utilities
+ */
+namespace FQuicTestUtils
+{
+
+	static IQuicNetworkMessagingExtension* GetQuicFeature()
+	{
+		TArray<INetworkMessagingExtension*> Features = IModularFeatures::Get()
+			.GetModularFeatureImplementations<INetworkMessagingExtension>(
+				INetworkMessagingExtension::ModularFeatureName);
+
+		for (const auto Feature : Features)
+		{
+			if (Feature->GetName() == "QuicMessaging")
+			{
+				IQuicNetworkMessagingExtension* QuicFeature
+					= static_cast<IQuicNetworkMessagingExtension*>(Feature);
+
+				return QuicFeature;
+			}
+		}
+
+		return nullptr;
+	}
+
+	template<typename T>
+	static TSharedRef<FQuicMetaMessageContext> WrapMessage(T* Message)
+	{
+		TSharedRef<FQuicMetaMessageContext> Context
+			= MakeShared<FQuicMetaMessageContext>(Message);
+
+		return Context;
+	}
+
+	static FString GetAuthPayload()
+	{
+		return TEXT("1234567890");
+	}
+
+};
+
+
+/**
+ * Test state
+ */
+class FQuicMessageTransportTestState
+    : public IMessageTransportHandler
+{
+    
+public:
+
+    FQuicMessageTransportTestState(
+        FAutomationTestBase& Test, const bool bInIsClient,
+        const TSharedRef<FQuicEndpointConfig> InEndpointConfig,
+		const bool bInEchoEndpoint, TArray<FIPv4Endpoint> InStaticEndpoints)
+		: bEchoEndpoint(bInEchoEndpoint)
+		, Transport(MakeShared<FQuicMessageTransport, ESPMode::ThreadSafe>(
+			bInIsClient, InEndpointConfig, MoveTemp(InStaticEndpoints)))
+    {
+    }
+
+    virtual ~FQuicMessageTransportTestState()
+	{
+	}
+
+public:
+
+    bool Start()
+    {
+		Transport->OnMetaMessageReceived().AddRaw(
+			this, &FQuicMessageTransportTestState::OnQuicMetaMessageReceived);
+
+		return Transport->StartTransport(*this);
+    }
+
+	void Stop() const
+	{
+		Transport->OnMetaMessageReceived().RemoveAll(this);
+
+		Transport->StopTransport();
+	}
+
+public:
+
+    void AddClient(const FIPv4Endpoint& RemoteEndpoint) const
+    {
+        Transport->AddStaticEndpoint(RemoteEndpoint);
+    }
+
+    void RemoveClient(const FIPv4Endpoint& RemoteEndpoint) const
+    {
+        Transport->RemoveStaticEndpoint(RemoteEndpoint);
+    }
+
+	void SetNodeAuthenticated(const FGuid& NodeId) const
+    {
+		Transport->SetNodeAuthenticated(NodeId);
+    }
+
+    void BroadcastMessage(
+		const TSharedRef<IMessageContext, ESPMode::ThreadSafe>& Context) const
+	{
+		const TArray<FGuid> Recipients;
+
+		Transport->TransportMessage(Context, Recipients);
+	}
+
+	void SendMessage(const FGuid& NodeId,
+		const TSharedRef<IMessageContext, ESPMode::ThreadSafe>& Context) const
+	{
+    	TArray<FGuid> Recipients;
+		Recipients.Add(NodeId);
+
+		Transport->TransportMessage(Context, Recipients);
+	}
+
+	void SendAuthMessage(const FGuid& NodeId, const FString& Payload) const
+    {
+		FQuicAuthMessage* AuthMessage
+			= FMessageEndpoint::MakeMessage<FQuicAuthMessage>();
+
+		AuthMessage->Payload = Payload;
+
+		const TSharedRef<FQuicMetaMessageContext> Context
+			= FQuicTestUtils::WrapMessage(AuthMessage);
+
+		Transport->TransportAuthMessage(Context, NodeId);
+    }
+
+	void TransportAuthMessage(
+		const TSharedRef<IMessageContext, ESPMode::ThreadSafe>& Context,
+		const FGuid& Recipient) const
+    {
+		Transport->TransportAuthMessage(Context, Recipient);
+    }
+
+private:
+
+	void OnQuicAuthMessageReceived(
+		const FGuid LocalNodeId, const TSharedRef<IMessageContext>& Context)
+	{
+		const FQuicAuthMessage* AuthMessage
+			= static_cast<const FQuicAuthMessage*>(Context->GetMessage());
+
+		if (!AuthMessage)
+		{
+			UE_LOG(LogTemp, Error,
+				TEXT("[QuicMessageTest::OnQuicMetaMessageReceived] "
+					"AuthMessage ptr not valid."));
+
+			return;
+		}
+
+		const bool bAuthSuccessful
+			= AuthMessage->Payload == FQuicTestUtils::GetAuthPayload();
+
+		// Set the node as authenticated
+		if (bAuthSuccessful)
+		{
+			SetNodeAuthenticated(LocalNodeId);
+		}
+
+		FQuicAuthResponseMessage* AuthResponseMessage
+			= FMessageEndpoint::MakeMessage<FQuicAuthResponseMessage>();
+
+		AuthResponseMessage->bAuthSuccessful = bAuthSuccessful;
+
+		const TSharedRef<FQuicMetaMessageContext> ResponseContext
+			= FQuicTestUtils::WrapMessage(AuthResponseMessage);
+
+		Transport->TransportAuthResponseMessage(ResponseContext, LocalNodeId);
+	}
+
+	void OnQuicAuthResponseMessageReceived(
+		const FGuid LocalNodeId, const TSharedRef<IMessageContext>& Context)
+	{
+		const FQuicAuthResponseMessage* AuthResponseMessage
+			= static_cast<const FQuicAuthResponseMessage*>(Context->GetMessage());
+
+		if (!AuthResponseMessage)
+		{
+			UE_LOG(LogTemp, Error,
+				TEXT("[QuicMessageTest::OnQuicMetaMessageReceived] "
+					"AuthResponseMessage ptr not valid."));
+
+			return;
+		}
+
+		bIsAuthenticated = AuthResponseMessage->bAuthSuccessful;
+
+		if (!bIsAuthenticated)
+		{
+			UE_LOG(LogTemp, Error,
+				TEXT("[QuicMessageTest::OnQuicAuthResponseMessageReceived] "
+					"Authentication failed. Reason: %s"), *AuthResponseMessage->Reason);
+		}
+	}
+
+	void OnQuicMetaMessageReceived(
+		const FGuid LocalNodeId, const TSharedRef<IMessageContext>& Context)
+	{
+		if (!Context->IsValid())
+		{
+			return;
+		}
+
+		if (Context->GetMessageTypeInfo()
+			->IsChildOf(FQuicAuthMessage::StaticStruct()))
+		{
+			OnQuicAuthMessageReceived(LocalNodeId, Context);
+		}
+		else
+		{
+			OnQuicAuthResponseMessageReceived(LocalNodeId, Context);
+		}
+	}
+
+public:
+
+    TArray<FIPv4Endpoint> GetEndpoints() const
+	{
+    	return Transport->GetKnownEndpoints();
+	}
+
+	TOptional<FGuid> GetNodeId(const FString& RemoteEndpoint) const
+    {
+		FIPv4Endpoint ParsedRemoteEndpoint;
+
+		if (!FIPv4Endpoint::Parse(RemoteEndpoint, ParsedRemoteEndpoint))
+		{
+			UE_LOG(LogTemp, Error,
+				TEXT("[QuicMessageTest::GetNodeId] Endpoint string could not be parsed."));
+
+			return TOptional<FGuid>();
+		}
+
+		return Transport->GetNodeId(ParsedRemoteEndpoint);
+    }
+
+	const TArray<FGuid>& GetDiscoveredNodes() const
+	{
+		return DiscoveredNodes;
+	}
+
+	const TArray<FGuid> GetLostNodes() const
+	{
+		return LostNodes;
+	}
+
+	int32 GetNumReceivedMessages() const
+	{
+		return NumReceivedMessages;
+	}
+
+	bool IsAuthenticated() const
+    {
+		return bIsAuthenticated;
+    }
+
+public:
+
+	//~ IMessageTransportHandler interface
+
+	virtual void DiscoverTransportNode(const FGuid& NodeId) override
+	{
+		DiscoveredNodes.Add(NodeId);
+	}
+
+	virtual void ForgetTransportNode(const FGuid& NodeId) override
+	{
+		LostNodes.Add(NodeId);
+	}
+
+	virtual void ReceiveTransportMessage(
+		const TSharedRef<IMessageContext, ESPMode::ThreadSafe>& Context,
+		const FGuid& NodeId) override
+	{
+		FPlatformAtomics::InterlockedIncrement(&NumReceivedMessages);
+
+		if (bEchoEndpoint)
+		{
+			SendMessage(NodeId, Context);
+		}
+	}
+
+private:
+
+	bool bEchoEndpoint = false;
+
+	bool bIsAuthenticated = false;
+
+    TArray<FGuid> DiscoveredNodes;
+
+	TArray<FGuid> LostNodes;
+
+	int32 NumReceivedMessages = 0;
+
+	TSharedRef<FQuicMessageTransport, ESPMode::ThreadSafe> Transport;
+
+};
+
+
+bool FQuicMessageTransportTest::RunTest(const FString& Parameters)
+{
+    const FIPv4Endpoint ServerEndpoint(FIPv4Address::Any, 9977);
+	const FIPv4Endpoint ClientEndpoint(FIPv4Address(127, 0, 0, 1), 9977);
+
+	const FGuid ServerGuid = FGuid::NewGuid();
+	const FGuid ClientGuid = FGuid::NewGuid();
+
+	const int32 NumTestMessages = 1000;
+
+	const TArray<FIPv4Endpoint> StaticEndpoints;
+
+    // Create the server transport config
+    TSharedRef<FQuicServerConfig> ServerConfig = MakeShared<FQuicServerConfig>();
+
+    {
+        ServerConfig->Endpoint = ServerEndpoint;
+        ServerConfig->LocalNodeId = ServerGuid;
+        ServerConfig->EncryptionMode = EEncryptionMode::Enabled;
+        ServerConfig->DiscoveryTimeoutSec = 0;
+        ServerConfig->Certificate = "";
+        ServerConfig->PrivateKey = "";
+        ServerConfig->MaxAuthenticationMessageSize = 8 * 1024;
+        ServerConfig->AuthenticationMode = EAuthenticationMode::Enabled;
+        ServerConfig->ConnCooldownMode = EConnectionCooldownMode::Enabled;
+        ServerConfig->ConnCooldownMaxAttempts = 5;
+        ServerConfig->ConnCooldownPeriodSec = 30;
+        ServerConfig->ConnCooldownSec = 30;
+        ServerConfig->ConnCooldownMaxSec = 3600;
+    }
+
+    // Setup server transport
+    const TSharedPtr<FQuicMessageTransportTestState> ServerTransport
+        = MakeShared<FQuicMessageTransportTestState>(
+            *this,                    // Test
+			false,                    // bInIsClient
+            ServerConfig,             // EndpointConfig
+			true,                     // bInEchoEndpoint
+            StaticEndpoints           // StaticEndpoints
+        );
+
+	// Create the client transport config
+	TSharedRef<FQuicClientConfig> ClientConfig = MakeShared<FQuicClientConfig>();
+
+	{
+		ClientConfig->Endpoint = ClientEndpoint;
+		ClientConfig->LocalNodeId = ClientGuid;
+		ClientConfig->EncryptionMode = EEncryptionMode::Enabled;
+		ClientConfig->DiscoveryTimeoutSec = 0;
+		ClientConfig->ClientVerificationMode = EQuicClientVerificationMode::Pass;
+	}
+
+    // Setup client one transport
+    const TSharedPtr<FQuicMessageTransportTestState> ClientTransport
+        = MakeShared<FQuicMessageTransportTestState>(
+            *this,                    // Test
+			true,                     // bInIsClient
+			ClientConfig,             // EndpointConfig
+			false,                    // bInEchoEndpoint
+            StaticEndpoints           // StaticEndpoints
+        );
+    
+    // Start server and one client
+    {
+        AddInfo(TEXT("Starting server ..."));
+
+        ServerTransport->Start();
+
+        AddInfo(TEXT("Starting client ..."));
+
+		ClientTransport->Start();
+		ClientTransport->AddClient(ClientEndpoint);
+
+		AddInfo(TEXT("Server and client should be up and running."));
+    }
+
+    /**
+     * Check server and client discovery
+     *
+     * This test checks whether the server and the client
+     * have discovered each other.
+     */
+    {
+        AddInfo(TEXT("Checking server and client discovery ..."));
+
+		const FTimespan MaxCheckDuration = FTimespan::FromSeconds(10);
+		const FDateTime CheckStart = FDateTime::UtcNow();
+
+		while (((ServerTransport->GetDiscoveredNodes().Num() < 1)
+			|| (ClientTransport->GetDiscoveredNodes().Num() < 1))
+			&& (FDateTime::UtcNow() < (CheckStart + MaxCheckDuration)))
+		{
+			FPlatformProcess::Sleep(0.0f);
+		}
+
+		TestTrue(
+			TEXT("Server must discover exactly one node."),
+			ServerTransport->GetDiscoveredNodes().Num() == 1);
+
+        TestTrue(
+            TEXT("Client must discover exactly one node."),
+			ClientTransport->GetDiscoveredNodes().Num() == 1);
+
+    	if (!ServerTransport->GetDiscoveredNodes().Contains(ClientGuid))
+        {
+			AddError(TEXT("Server did not discover client."));
+        }
+
+    	if (!ClientTransport->GetDiscoveredNodes().Contains(ServerGuid))
+        {
+			AddError(TEXT("Client did not discover server."));
+        }
+
+        if (HasAnyErrors())
+        {
+			return false;
+        }
+
+		AddInfo(TEXT("Client and server discovered each other."));
+    }
+
+	/**
+	 * Authenticate the client with the server
+	 *
+	 * Since the QuicServer has endpoint authentication enabled,
+	 * the QuicClient needs to send an authentication message that
+	 * the server will check and if authenticated, the server will allow
+	 * the client to send normal data flow messages.
+	 */
+    {
+		AddInfo(TEXT("Sending authentication message ..."));
+
+		ClientTransport->SendAuthMessage(ServerGuid, FQuicTestUtils::GetAuthPayload());
+
+		AddInfo(TEXT("Sent authentication message. Waiting for response ..."));
+
+		const FTimespan MaxCheckDuration = FTimespan::FromSeconds(10);
+		const FDateTime CheckStart = FDateTime::UtcNow();
+
+		while ((!ClientTransport->IsAuthenticated())
+			&& (FDateTime::UtcNow() < (CheckStart + MaxCheckDuration)))
+		{
+			FPlatformProcess::Sleep(0.0f);
+		}
+
+		if (!ClientTransport->IsAuthenticated())
+		{
+			AddError(TEXT("Client was not authenticated."));
+		}
+
+		if (HasAnyErrors())
+		{
+			return false;
+		}
+
+		AddInfo(TEXT("Client is authenticated and can send normal messages now."));
+    }
+
+	/**
+	 * Send small test messages from client to server
+	 *
+	 * This test will send a lot of small messages to the server,
+	 * which will in turn send them back to the client to ensure
+	 * that both directions of communication are working as intended.
+	 */
+    {
+		const int32 MessageSize = 1024;
+
+		AddInfo(TEXT("Sending small test messages to server ..."));
+
+		for (int32 Count = 0; Count < NumTestMessages; ++Count)
+		{
+			FQuicMockMessage* MockMessage
+				= FMessageEndpoint::MakeMessage<FQuicMockMessage>(MessageSize);
+
+			const TSharedRef<FQuicMetaMessageContext> Context
+				= FQuicTestUtils::WrapMessage(MockMessage);
+
+			ClientTransport->SendMessage(ServerGuid, Context);
+		}
+
+		AddInfo(TEXT("Test messages sent to server. Waiting to receive them ..."));
+
+		const FTimespan MaxCheckDuration = FTimespan::FromSeconds(30);
+		const FDateTime CheckStart = FDateTime::UtcNow();
+		
+		while (((ServerTransport->GetNumReceivedMessages() < NumTestMessages)
+			|| (ClientTransport->GetNumReceivedMessages() < NumTestMessages))
+			&& (FDateTime::UtcNow() < (CheckStart + MaxCheckDuration)))
+		{
+			FPlatformProcess::Sleep(0.0f);
+		}
+
+		if (ServerTransport->GetNumReceivedMessages() != NumTestMessages)
+		{
+			AddError(TEXT("Server did not receive all test messages."));
+		}
+
+		// Waiting for the response from the echo server
+		if (ClientTransport->GetNumReceivedMessages() != NumTestMessages)
+		{
+			AddError(TEXT("Client did not receive all test messages."));
+		}
+
+		if (HasAnyErrors())
+		{
+			return false;
+		}
+
+		AddInfo(TEXT("All endpoints received all small test messages."));
+    }
+
+	/**
+	 * Send a large test message from client to server
+	 *
+	 * This test will send one large message to the server,
+	 * which will in turn send it back to the client to ensure
+	 * that both directions of communication are working as intended.
+	 */
+    {
+		const int32 MessageSize = 100 * 1024 * 1024; // 100 MB
+		const int32 ExpectedReceivedMessages = NumTestMessages + 1;
+
+		AddInfo(TEXT("Sending large test message to the server ..."));
+
+		FQuicMockMessage* MockMessage
+			= FMessageEndpoint::MakeMessage<FQuicMockMessage>(MessageSize);
+
+		const TSharedRef<FQuicMetaMessageContext> Context
+			= FQuicTestUtils::WrapMessage(MockMessage);
+
+		ClientTransport->SendMessage(ServerGuid, Context);
+
+		AddInfo(TEXT("Large test message sent to the server. "
+			   "Waiting to receive it on all endpoints ..."));
+
+		const FTimespan MaxCheckDuration = FTimespan::FromSeconds(60);
+		const FDateTime CheckStart = FDateTime::UtcNow();
+
+		while (((ServerTransport->GetNumReceivedMessages() < ExpectedReceivedMessages)
+			|| (ClientTransport->GetNumReceivedMessages() < ExpectedReceivedMessages))
+			&& (FDateTime::UtcNow() < (CheckStart + MaxCheckDuration)))
+		{
+			FPlatformProcess::Sleep(0.0f);
+		}
+
+		if (ServerTransport->GetNumReceivedMessages() != ExpectedReceivedMessages)
+		{
+			AddError(TEXT("Server did not receive large test message."));
+		}
+
+		// Waiting for the response from the echo server
+		if (ClientTransport->GetNumReceivedMessages() != ExpectedReceivedMessages)
+		{
+			AddError(TEXT("Client did not receive large test message."));
+		}
+
+		if (HasAnyErrors())
+		{
+			return false;
+		}
+
+		AddInfo(TEXT("All endpoints received the large test message."));
+    }
+
+	// Shutdown original client to start a new, unauthenticated one
+    {
+		ClientTransport->RemoveClient(ClientEndpoint);
+		ClientTransport->Stop();
+
+		FPlatformProcess::Sleep(3.0f);
+    }
+
+	// Create another client Guid and adapt the client config
+	const FGuid AnotherClientGuid = FGuid::NewGuid();
+	ClientConfig->LocalNodeId = AnotherClientGuid;
+
+	// Setup another client transport
+	const TSharedPtr<FQuicMessageTransportTestState> ClientTransportNoAuth
+		= MakeShared<FQuicMessageTransportTestState>(
+			*this,                    // Test
+			true,                     // bInIsClient
+			ClientConfig,             // EndpointConfig
+			false,                    // bInEchoEndpoint
+			StaticEndpoints           // StaticEndpoints
+		);
+
+	/**
+	 * Connect and send message with unauthenticated client
+	 *
+	 * This test checks whether an unauthenticated client can send normal
+	 * data messages to the server. That should not work and would throw an error
+	 * on the server side, which is caught as to not make the test fail.
+	 */
+    {
+		AddInfo(TEXT("Connecting with another client ..."));
+
+		ClientTransportNoAuth->Start();
+		ClientTransportNoAuth->AddClient(ClientEndpoint);
+
+		FPlatformProcess::Sleep(3.0f);
+
+		if (!ServerTransport->GetDiscoveredNodes().Contains(AnotherClientGuid))
+		{
+			AddError(TEXT("Server did not discover new client."));
+		}
+
+		const int32 MessageSize = 2048;
+		const int32 ExpectedReceivedMessages = NumTestMessages + 2;
+
+		AddExpectedError(TEXT("Unauthenticated handler sending data"));
+
+		AddInfo(TEXT("Sending unauthenticated test message to the server ..."));
+
+		FQuicMockMessage* MockMessage
+			= FMessageEndpoint::MakeMessage<FQuicMockMessage>(MessageSize);
+
+		const TSharedRef<FQuicMetaMessageContext> Context
+			= FQuicTestUtils::WrapMessage(MockMessage);
+
+		ClientTransportNoAuth->SendMessage(ServerGuid, Context);
+
+		AddInfo(TEXT("Unauthenticated test message sent to the server."));
+
+		FPlatformProcess::Sleep(3.0f);
+
+		TestFalse(
+			TEXT("Server shouldn't receive unauthenticated message."),
+			ServerTransport->GetNumReceivedMessages() == ExpectedReceivedMessages);
+
+		if (HasAnyErrors())
+		{
+			return false;
+		}
+
+		AddInfo(TEXT("Server did not receive unauthenticated message."));
+    }
+
+	return true;
+}
+
+
+void EmptyLinkFunctionForStaticInitializationQuicMessageTransportTest()
+{
+	// This function exists to prevent the object file containing this test
+	// from being excluded by the linker, because it has no publicly referenced symbols.
+}
+
diff --git a/Engine/Plugins/Experimental/QuicMessaging/Source/QuicMessaging/Private/Tests/QuicMessagingTestTypes.h b/Engine/Plugins/Experimental/QuicMessaging/Source/QuicMessaging/Private/Tests/QuicMessagingTestTypes.h
new file mode 100644
--- /dev/null
+++ b/Engine/Plugins/Experimental/QuicMessaging/Source/QuicMessaging/Private/Tests/QuicMessagingTestTypes.h
@@ -0,0 +1,29 @@
+// Copyright Epic Games, Inc. All Rights Reserved.
+
+#pragma once
+
+#include "CoreMinimal.h"
+#include "IMessageContext.h"
+#include "QuicMessagingTestTypes.generated.h"
+
+
+USTRUCT()
+struct FQuicMockMessage
+{
+
+	GENERATED_BODY()
+
+	UPROPERTY(meta = (IgnoreForMemberInitializationTest))
+	TArray<uint8> Data;
+
+	FQuicMockMessage()
+	{
+		Data.AddUninitialized(64);
+	}
+
+	FQuicMockMessage(int32 DataSize)
+	{
+		Data.AddUninitialized(DataSize);
+	}
+
+};
diff --git a/Engine/Plugins/Experimental/QuicMessaging/Source/QuicMessaging/Private/Transport/QuicMessageProcessor.cpp b/Engine/Plugins/Experimental/QuicMessaging/Source/QuicMessaging/Private/Transport/QuicMessageProcessor.cpp
new file mode 100644
--- /dev/null
+++ b/Engine/Plugins/Experimental/QuicMessaging/Source/QuicMessaging/Private/Transport/QuicMessageProcessor.cpp
@@ -0,0 +1,345 @@
+// Copyright Epic Games, Inc. All Rights Reserved.
+
+#include "Transport/QuicMessageProcessor.h"
+#include "QuicMessagingPrivate.h"
+
+#include "HAL/Event.h"
+#include "HAL/RunnableThread.h"
+#include "IMessageAttachment.h"
+#include "Serialization/ArrayReader.h"
+#include "Serialization/ArrayWriter.h"
+#include "Math/UnrealMathUtility.h"
+#include "UObject/Class.h"
+
+#include "Backends/CborStructSerializerBackend.h"
+#include "StructSerializer.h"
+#include "Backends/CborStructDeserializerBackend.h"
+#include "StructDeserializer.h"
+
+#include "IMessageContext.h"
+#include "INetworkMessagingExtension.h"
+
+#include "QuicSerializedMessage.h"
+#include "QuicDeserializedMessage.h"
+#include "QuicMessagingSettings.h"
+
+#include "QuicEndpointManager.h"
+
+
+namespace UE::Private::MessageProcessor
+{
+
+	FOnOutboundTransferDataUpdated& OnOutboundUpdated()
+	{
+		static FOnOutboundTransferDataUpdated OnTransferUpdated;
+		return OnTransferUpdated;
+	}
+
+	FOnInboundTransferDataUpdated& OnInboundUpdated()
+	{
+		static FOnInboundTransferDataUpdated OnTransferUpdated;
+		return OnTransferUpdated;
+	}
+
+}
+
+
+FQuicMessageProcessor::FQuicMessageProcessor()
+    : bStopping(false)
+    , bIsInitialized(false)
+    , MessageFormat(GetDefault<UQuicMessagingSettings>()->MessageFormat)
+	, Thread(nullptr)
+	, SerializeWorkEvent(nullptr)
+	, DeserializeWorkEvent(nullptr)
+	, bUpdatingSerialization(false)
+	, PendingSerializeMessages(TMap<uint32, FQuicSerializedMessagePtr>())
+	, PendingDeserializeMessages(TMap<uint32, FQuicDeserializedMessagePtr>())
+	, SerializeMessagesQueue(TQueue<FQuicSerializedMessagePtr>())
+	, DeserializeMessagesQueue(TQueue<FQuicDeserializedMessagePtr>())
+{
+	SerializeWorkEvent = MakeShareable(FPlatformProcess::GetSynchEventFromPool(), [](FEvent* EventToDelete)
+    {
+        FPlatformProcess::ReturnSynchEventToPool(EventToDelete);
+    });
+
+	DeserializeWorkEvent = MakeShareable(FPlatformProcess::GetSynchEventFromPool(), [](FEvent* EventToDelete)
+    {
+        FPlatformProcess::ReturnSynchEventToPool(EventToDelete);
+    });
+
+    Thread = FRunnableThread::Create(this, *FString::Printf(TEXT("FQuicMessageProcessor-%d"), FMath::RandRange(0, 10)), 128 * 4096,
+        TPri_AboveNormal, FPlatformAffinity::GetPoolThreadMask());
+
+	UE_LOG(LogQuicMessaging, Verbose, TEXT("[MessageProcessor] Started message processor."));
+}
+
+
+FQuicMessageProcessor::~FQuicMessageProcessor()
+{
+    Thread->Kill(true);
+    delete Thread;
+    Thread = nullptr;
+}
+
+
+void FQuicMessageProcessor::ProcessInboundMessage(const FInboundMessage InboundMessage)
+{
+    if (bStopping)
+    {
+        return;
+    }
+
+	TSharedPtr<FQuicDeserializedMessage, ESPMode::ThreadSafe> DeserializedMessage
+		= MakeShared<FQuicDeserializedMessage, ESPMode::ThreadSafe>(InboundMessage);
+
+	DeserializeMessagesQueue.Enqueue(DeserializedMessage);
+}
+
+
+bool FQuicMessageProcessor::ProcessOutboundMessage(
+    const TSharedRef<IMessageContext, ESPMode::ThreadSafe> MessageContext,
+    const TArray<FGuid>& Recipients, const EQuicMessageType MessageType)
+{
+    if (bStopping)
+    {
+        return false;
+    }
+
+	FQuicPayloadPtr DataPtr = MakeShared<TArray<uint8>, ESPMode::ThreadSafe>();
+
+    TSharedPtr<FQuicSerializedMessage, ESPMode::ThreadSafe> SerializedMessage
+        = MakeShared<FQuicSerializedMessage, ESPMode::ThreadSafe>(
+		MessageFormat, MessageContext->GetFlags(), Recipients, MessageType, DataPtr, MessageContext);
+
+	SerializeMessagesQueue.Enqueue(SerializedMessage);
+
+	return true;
+}
+
+
+void FQuicMessageProcessor::ConsumeMessages()
+{
+	if (bStopping || bConsumingMessages)
+	{
+		return;
+	}
+
+	bConsumingMessages = true;
+
+	const double CurrentTime = FPlatformTime::Seconds();
+
+	if (CurrentTime < (LastConsumed + ConsumeInterval.GetTotalSeconds()))
+	{
+		bConsumingMessages = false;
+		return;
+	}
+
+	if (SerializeMessagesQueue.IsEmpty() && DeserializeMessagesQueue.IsEmpty())
+	{
+		bConsumingMessages = false;
+		return;
+	}
+
+	TSharedPtr<FQuicSerializedMessage, ESPMode::ThreadSafe> SerializedMessage;
+
+	while (!bStopping && SerializeMessagesQueue.Dequeue(SerializedMessage))
+	{
+		uint32 SerializeKey = SerializeTaskId;
+		SerializeTaskId++;
+
+		PendingSerializeMessages.Add(SerializeKey, SerializedMessage);
+
+		TGraphTask<FQuicSerializeMessageTask>::CreateTask()
+			.ConstructAndDispatchWhenReady(
+				SerializedMessage->GetMessageContext(),
+				SerializedMessage, SerializeWorkEvent);
+	}
+
+	TSharedPtr<FQuicDeserializedMessage, ESPMode::ThreadSafe> DeserializedMessage;
+
+	while (!bStopping && DeserializeMessagesQueue.Dequeue(DeserializedMessage))
+	{
+		uint32 DeserializeKey = DeserializeTaskId;
+		DeserializeTaskId++;
+
+		PendingDeserializeMessages.Add(DeserializeKey, DeserializedMessage);
+
+		TGraphTask<FQuicDeserializeMessageTask>::CreateTask()
+			.ConstructAndDispatchWhenReady(
+				DeserializedMessage->GetInboundMessage().UnserializedMessage,
+				DeserializedMessage, DeserializeWorkEvent);
+	}
+
+	LastConsumed = FPlatformTime::Seconds();
+
+	bConsumingMessages = false;
+}
+
+
+void FQuicMessageProcessor::UpdateSerialization()
+{
+	if (bStopping || bUpdatingSerialization)
+	{
+		return;
+	}
+
+	bUpdatingSerialization = true;
+
+	const double CurrentTime = FPlatformTime::Seconds();
+
+	if (CurrentTime < (LastSerialization + SerializationInterval.GetTotalSeconds()))
+	{
+		bUpdatingSerialization = false;
+		return;
+	}
+
+	for (auto SerializeIt = PendingSerializeMessages.CreateIterator(); SerializeIt; ++SerializeIt)
+	{
+		TSharedPtr<FQuicSerializedMessage, ESPMode::ThreadSafe> SerializedMessage = SerializeIt->Value;
+
+		if (!SerializedMessage.IsValid())
+		{
+			UE_LOG(LogQuicMessaging, Error, TEXT("[MessageProcessor] Serialized message ptr invalid."));
+
+			SerializeIt.RemoveCurrent();
+
+			continue;
+		}
+
+		if (SerializedMessage->GetState() == EQuicSerializedMessageState::Complete)
+		{
+			MessageSerializedDelegate.ExecuteIfBound(SerializedMessage);
+
+			FOnOutboundTransferDataUpdated& OutboundUpdated = UE::Private::MessageProcessor::OnOutboundUpdated();
+
+			SerializeIt.RemoveCurrent();
+		}
+
+		if (SerializedMessage->GetState() == EQuicSerializedMessageState::Invalid)
+		{
+			UE_LOG(LogQuicMessaging, Error, TEXT("[MessageProcessor] Message serialize failed."));
+
+			SerializeIt.RemoveCurrent();
+		}
+	}
+
+	for (auto DeserializeIt = PendingDeserializeMessages.CreateIterator(); DeserializeIt; ++DeserializeIt)
+	{
+		TSharedPtr<FQuicDeserializedMessage, ESPMode::ThreadSafe> DeserializedMessage = DeserializeIt->Value;
+
+		if (!DeserializedMessage.IsValid())
+		{
+			UE_LOG(LogQuicMessaging, Error, TEXT("[MessageProcessor] Deserialized message pointer invalid."));
+
+			DeserializeIt.RemoveCurrent();
+
+			continue;
+		}
+
+		if (DeserializedMessage->GetState() == EQuicDeserializedMessageState::Complete)
+		{
+			MessageDeserializedDelegate.ExecuteIfBound(DeserializedMessage,
+				DeserializedMessage->GetInboundMessage().MessageHeader.SenderId);
+
+			FOnInboundTransferDataUpdated& InboundUpdated = UE::Private::MessageProcessor::OnInboundUpdated();
+
+			DeserializeIt.RemoveCurrent();
+		}
+
+		if (DeserializedMessage->GetState() == EQuicDeserializedMessageState::Invalid)
+		{
+			DeserializeIt.RemoveCurrent();
+		}
+	}
+
+	LastSerialization = FPlatformTime::Seconds();
+
+	bUpdatingSerialization = false;
+
+}
+
+
+FQuicPayloadPtr FQuicMessageProcessor::SerializeMessageHeader(FMessageHeader& InMessageHeader) const
+{
+	FQuicPayloadPtr SerializedData = MakeShared<TArray<uint8>, ESPMode::ThreadSafe>();
+
+	SerializedData->Reserve(QUIC_MESSAGE_HEADER_SIZE);
+
+	FMemoryWriter Archive(*SerializedData.Get(), true);
+
+	uint8 MessageType = static_cast<uint8>(InMessageHeader.MessageType);
+	Archive << MessageType;
+
+	Archive << InMessageHeader.RecipientId;
+
+	Archive << InMessageHeader.SenderId;
+	
+	Archive << InMessageHeader.SerializedMessageSize;
+
+	SerializedData->AddDefaulted(QUIC_MESSAGE_HEADER_SIZE - Archive.TotalSize());
+
+	return SerializedData;
+}
+
+
+FMessageHeader FQuicMessageProcessor::DeserializeMessageHeader(FQuicPayloadPtr InHeaderData) const
+{
+	FMessageHeader MessageHeader;
+
+	FMemoryReader Archive(*InHeaderData);
+
+	uint8 MessageType;
+	Archive << MessageType;
+	MessageHeader.MessageType = static_cast<EQuicMessageType>(MessageType);
+
+	FGuid RecipientId;
+	Archive << RecipientId;
+	MessageHeader.RecipientId = RecipientId;
+
+	FGuid SenderId;
+	Archive << SenderId;
+	MessageHeader.SenderId = SenderId;
+
+	uint32 SerializedMessageSize;
+	Archive << SerializedMessageSize;
+	MessageHeader.SerializedMessageSize = SerializedMessageSize;
+
+	return MessageHeader;
+}
+
+
+void FQuicMessageProcessor::WaitAsyncTaskCompletion()
+{
+    // Stop the processor thread
+    Stop();
+}
+
+
+FSingleThreadRunnable* FQuicMessageProcessor::GetSingleThreadInterface()
+{
+	return this;
+}
+
+
+uint32 FQuicMessageProcessor::Run()
+{
+    while (!bStopping)
+    {
+		ConsumeMessages();
+		UpdateSerialization();
+    }
+
+    return 0;
+}
+
+
+void FQuicMessageProcessor::Stop()
+{
+	bStopping = true;
+	SerializeWorkEvent->Trigger();
+	DeserializeWorkEvent->Trigger();
+}
+
+
+
+
diff --git a/Engine/Plugins/Experimental/QuicMessaging/Source/QuicMessaging/Private/Transport/QuicMessageProcessor.h b/Engine/Plugins/Experimental/QuicMessaging/Source/QuicMessaging/Private/Transport/QuicMessageProcessor.h
new file mode 100644
--- /dev/null
+++ b/Engine/Plugins/Experimental/QuicMessaging/Source/QuicMessaging/Private/Transport/QuicMessageProcessor.h
@@ -0,0 +1,248 @@
+// Copyright Epic Games, Inc. All Rights Reserved.
+
+#pragma once
+
+#include "QuicMessagingPrivate.h"
+#include "CoreTypes.h"
+#include "Containers/Map.h"
+#include "Containers/Queue.h"
+#include "HAL/Runnable.h"
+#include "IMessageTransport.h"
+#include "Interfaces/IPv4/IPv4Endpoint.h"
+#include "Misc/DateTime.h"
+#include "Misc/Guid.h"
+#include "Misc/SingleThreadRunnable.h"
+#include "Misc/Timespan.h"
+#include "Templates/SharedPointer.h"
+#include "INetworkMessagingExtension.h"
+
+#include "QuicSerializationTasks.h"
+
+#include "QuicFlags.h"
+#include "QuicMessages.h"
+
+
+class FArrayReader;
+class FEvent;
+class FRunnableThread;
+
+class FQuicEndpointManager;
+struct FInboundMessage;
+
+class FQuicSerializedMessage;
+class FQuicDeserializedMessage;
+class IMessageAttachment;
+class IMessageContext;
+enum class EQuicMessageFormat : uint8;
+
+
+/**
+ * @note: The delegate macro recognizes the ThreadSafe TSharedPtr as two arguments.
+ */
+using FQuicDeserializedMessagePtr = TSharedPtr<FQuicDeserializedMessage, ESPMode::ThreadSafe>;
+using FQuicSerializedMessagePtr = TSharedPtr<FQuicSerializedMessage, ESPMode::ThreadSafe>;
+
+
+namespace UE::Private::MessageProcessor
+{
+
+/**
+ * Global delegate for handling outbound processor updates.
+ */
+FOnOutboundTransferDataUpdated& OnOutboundUpdated();
+
+/**
+ * Global delegate for handling inbound processor updates.
+ */
+FOnInboundTransferDataUpdated& OnInboundUpdated();
+}
+
+
+/**
+ * Implements a message processor for QUIC messages.
+ */
+class FQuicMessageProcessor
+	: public FRunnable
+	, private FSingleThreadRunnable
+{
+
+public:
+
+	/**
+	 * Creates and initializes a new message processor.
+	 */
+	FQuicMessageProcessor();
+
+	/** Virtual destructor. */
+	virtual ~FQuicMessageProcessor();
+
+public:
+
+	/**
+	 * Prepares and enqueues an inbound message for deserialization.
+	 *
+	 * @param InboundMessage The inbound message
+	 */
+	void ProcessInboundMessage(const FInboundMessage InboundMessage);
+
+	/**
+	 * Prepares and enqueues an outbound message for serialization.
+	 *
+	 * @param MessageContext The message to serialize
+	 * @param Recipients The recipient nodes ids
+	 * @param MessageType The message type
+	 * 
+	 * @return Whether the message could be enqueued
+	 */
+	bool ProcessOutboundMessage(
+		const TSharedRef<IMessageContext, ESPMode::ThreadSafe> MessageContext,
+		const TArray<FGuid>& Recipients, const EQuicMessageType MessageType);
+
+private:
+
+	/**
+	 * Consume messages and create task to serialize and deserialize messages.
+	 */
+	void ConsumeMessages();
+
+public:
+
+	/**
+	 * Updates message serialization tasks.
+	 */
+	void UpdateSerialization();
+
+	/**
+	 * Serializes a message header.
+	 *
+	 * @param InMessageHeader The message header to serialize
+	 * @return The serialized message header
+	 */
+	FQuicPayloadPtr SerializeMessageHeader(FMessageHeader& InMessageHeader) const;
+
+	/**
+	 * Deserializes a message header.
+	 *
+	 * @param InHeaderData The message header data to deserialize
+	 * @return The deserialized message header
+	 */
+	FMessageHeader DeserializeMessageHeader(FQuicPayloadPtr InHeaderData) const;
+
+public:
+
+	/**
+	 * Waits for all serialization tasks fired by this processor to complete.
+	 * Expected to be called when the application exit to prevent
+	 * serialized (UStruct) object to being use after the UObject system is shutdown.
+	 */
+	void WaitAsyncTaskCompletion();
+
+public:
+
+	/**
+	 * Returns a delegate that is executed when a message has been deserialized.
+	 *
+	 * @return The delegate.
+	 */
+	DECLARE_DELEGATE_TwoParams(FOnMessageDeserialized, FQuicDeserializedMessagePtr /*DeserializedMessage*/, FGuid /*NodeId*/)
+	FOnMessageDeserialized& OnMessageDeserialized()
+	{
+		return MessageDeserializedDelegate;
+	}
+
+	/**
+	 * Returns a delegate that is executed when a message has been serialized.
+	 *
+	 * @return The delegate.
+	 */
+	DECLARE_DELEGATE_OneParam(FOnMessageSerialized, FQuicSerializedMessagePtr /*SerializedMessage*/)
+	FOnMessageSerialized& OnMessageSerialized()
+	{
+		return MessageSerializedDelegate;
+	}
+
+public:
+
+	//~ FRunnable interface
+
+	virtual FSingleThreadRunnable* GetSingleThreadInterface() override;
+	virtual bool Init() override { return true; };
+	virtual uint32 Run() override;
+	virtual void Stop() override;
+	virtual void Exit() override { }
+
+protected:
+
+	//~ FSingleThreadRunnable interface
+
+	virtual void Tick() override { };
+
+private:
+
+	/** Holds a flag indicating that the thread is stopping. */
+	bool bStopping;
+
+	/** Holds a flag indicating if the processor is initialized. */
+	bool bIsInitialized;
+
+	/** The configured message format (from UQuicMessagingSettings). */
+	EQuicMessageFormat MessageFormat;
+
+	/** Holds the thread object. */
+	FRunnableThread* Thread;
+
+	/** Holds an event signaling that outbound messages need to be processed. */
+	TSharedPtr<FEvent, ESPMode::ThreadSafe> SerializeWorkEvent;
+
+	/** Holds an event signaling that inbound messages need to be processed. */
+	TSharedPtr<FEvent, ESPMode::ThreadSafe> DeserializeWorkEvent;
+
+private:
+
+	/** Holds a flag indicating if the processor is updating serialization. */
+	bool bUpdatingSerialization;
+
+	/** Holds the interval for message de/serialization. */
+	const FTimespan SerializationInterval = 10000;
+
+	/** Holds timestamp of last de/serialization. */
+	double LastSerialization = FPlatformTime::Seconds();
+
+	/** Holds the current serialize task id. */
+	uint32 SerializeTaskId = 0;
+
+	/** Holds the current deserialize task id. */
+	uint32 DeserializeTaskId = 0;
+
+	/** Holds pending serialized messages. */
+	TMap<uint32, FQuicSerializedMessagePtr> PendingSerializeMessages;
+
+	/** Holds pending deserialized messages. */
+	TMap<uint32, FQuicDeserializedMessagePtr> PendingDeserializeMessages;
+
+private:
+
+	/** Flag indicating if messages are currently being consumed.*/
+	bool bConsumingMessages = false;
+
+	/** Holds the interval for message consume. */
+	const FTimespan ConsumeInterval = 10000;
+
+	/** Holds the timestamp of when messages have last been consumed. */
+	double LastConsumed = FPlatformTime::Seconds();
+
+	/** Queue that holds serialized messages. */
+	TQueue<FQuicSerializedMessagePtr> SerializeMessagesQueue;
+
+	/** Queue that holds deserialized messages. */
+	TQueue<FQuicDeserializedMessagePtr> DeserializeMessagesQueue;
+
+private:
+
+	/** Holds a delegate to be invoked when a message has been deserialized. */
+	FOnMessageDeserialized MessageDeserializedDelegate;
+
+	/** Holds a delegate to be invoked when a message has been serialized. */
+	FOnMessageSerialized MessageSerializedDelegate;
+	
+};
diff --git a/Engine/Plugins/Experimental/QuicMessaging/Source/QuicMessaging/Private/Transport/QuicMessageTransport.cpp b/Engine/Plugins/Experimental/QuicMessaging/Source/QuicMessaging/Private/Transport/QuicMessageTransport.cpp
new file mode 100644
--- /dev/null
+++ b/Engine/Plugins/Experimental/QuicMessaging/Source/QuicMessaging/Private/Transport/QuicMessageTransport.cpp
@@ -0,0 +1,455 @@
+// Copyright Epic Games, Inc. All Rights Reserved.
+
+#include "Transport/QuicMessageTransport.h"
+
+#include "HAL/RunnableThread.h"
+#include "IMessageContext.h"
+#include "IMessageTransportHandler.h"
+#include "Misc/Guid.h"
+#include "Serialization/ArrayReader.h"
+#include "Async/Async.h"
+
+#include "QuicMessagingPrivate.h"
+#include "QuicMessages.h"
+#include "QuicEndpointConfig.h"
+#include "Transport/QuicMessageProcessor.h"
+
+#include "QuicEndpointManager.h"
+
+
+FQuicMessageTransport::FQuicMessageTransport(const bool bInIsClient,
+	const TSharedRef<FQuicEndpointConfig> InEndpointConfig,
+	TArray<FIPv4Endpoint> InStaticEndpoints)
+	: EndpointConfig(InEndpointConfig)
+    , StaticEndpoints(MoveTemp(InStaticEndpoints))
+{
+    EndpointMode = (bInIsClient) ? EEndpointMode::Client : EEndpointMode::Server;
+}
+
+
+FQuicMessageTransport::~FQuicMessageTransport()
+{
+	StopTransport();
+}
+
+
+void FQuicMessageTransport::OnAppPreExit()
+{
+	// Remove any bound delegates. It's no longer relevant for us to
+	// send a transport error when we are in the shutdown phase.
+	TransportErrorDelegate.Unbind();
+
+	if (MessageProcessor.IsValid())
+	{
+		MessageProcessor->WaitAsyncTaskCompletion();
+	}
+}
+
+
+void FQuicMessageTransport::AddStaticEndpoint(const FIPv4Endpoint& InEndpoint)
+{
+	if (EndpointManager == nullptr)
+	{
+		return;
+	}
+
+	if (EndpointMode == EEndpointMode::Server)
+	{
+		return;
+	}
+
+    EndpointManager->AddClient(InEndpoint);
+}
+
+
+void FQuicMessageTransport::RemoveStaticEndpoint(const FIPv4Endpoint& InEndpoint)
+{
+	if (EndpointManager != nullptr)
+	{
+		EndpointManager->RemoveClient(InEndpoint);
+	}
+}
+
+
+TArray<FIPv4Endpoint> FQuicMessageTransport::GetListeningAddresses() const
+{
+	TArray<FIPv4Endpoint> Listeners;
+
+	Listeners.Add(EndpointConfig->Endpoint);
+
+	return Listeners;
+}
+
+
+TArray<FIPv4Endpoint> FQuicMessageTransport::GetKnownEndpoints() const
+{
+	if (EndpointManager)
+	{
+		return EndpointManager->GetKnownEndpoints();
+	}
+
+	return {};
+}
+
+
+FMessageTransportStatistics FQuicMessageTransport::GetLatestStatistics(FGuid NodeId) const
+{
+	if (EndpointManager)
+	{
+		return EndpointManager->GetStats(NodeId);
+	}
+
+	return {};
+}
+
+
+void FQuicMessageTransport::SetMaxAuthenticationMessageSize(const uint32 MaxBytes) const
+{
+	if (!EndpointManager)
+	{
+		return;
+	}
+
+	if (EndpointMode != EEndpointMode::Server)
+	{
+		return;
+	}
+
+	const TSharedRef<FQuicServerConfig> ServerEndpointConfig
+		= StaticCastSharedRef<FQuicServerConfig, FQuicEndpointConfig>(EndpointConfig);
+
+	ServerEndpointConfig->MaxAuthenticationMessageSize = MaxBytes;
+}
+
+
+TOptional<FGuid> FQuicMessageTransport::GetNodeId(const FIPv4Endpoint& RemoteEndpoint) const
+{
+	if (!EndpointManager)
+	{
+		return TOptional<FGuid>();
+	}
+
+	return EndpointManager->GetKnownNodeId(RemoteEndpoint);
+}
+
+
+bool FQuicMessageTransport::IsNodeAuthenticated(const FGuid& NodeId) const
+{
+    if (EndpointManager == nullptr)
+    {
+        return false;
+    }
+
+    return EndpointManager->IsEndpointAuthenticated(NodeId);
+}
+
+
+void FQuicMessageTransport::SetNodeAuthenticated(const FGuid& NodeId)
+{
+    if (EndpointManager == nullptr)
+    {
+        return;
+    }
+
+	EndpointManager->SetEndpointAuthenticated(NodeId);
+}
+
+
+void FQuicMessageTransport::DisconnectNode(const FGuid& NodeId) const
+{
+	if (!EndpointManager)
+	{
+		return;
+	}
+
+	EndpointManager->DisconnectNode(NodeId);
+}
+
+
+void FQuicMessageTransport::SetConnectionCooldown(
+	const bool bEnabled, const uint32 MaxAttempts, const uint32 PeriodSeconds,
+	const uint32 CooldownSeconds, const uint32 CooldownMaxSeconds) const
+{
+	if (!EndpointManager)
+	{
+		return;
+	}
+
+	if (EndpointMode != EEndpointMode::Server)
+	{
+		return;
+	}
+
+	const TSharedRef<FQuicServerConfig> ServerEndpointConfig
+		= StaticCastSharedRef<FQuicServerConfig, FQuicEndpointConfig>(EndpointConfig);
+
+	ServerEndpointConfig->ConnCooldownMode = (bEnabled)
+		? EConnectionCooldownMode::Enabled : EConnectionCooldownMode::Disabled;
+
+	ServerEndpointConfig->ConnCooldownMaxAttempts = MaxAttempts;
+	ServerEndpointConfig->ConnCooldownPeriodSec = PeriodSeconds;
+	ServerEndpointConfig->ConnCooldownSec = CooldownSeconds;
+	ServerEndpointConfig->ConnCooldownMaxSec = CooldownMaxSeconds;
+}
+
+
+bool FQuicMessageTransport::RestartTransport()
+{
+	UE_LOG(LogQuicMessaging, Verbose, TEXT("[MessageTransport] Restarting transport ..."));
+
+	IMessageTransportHandler* Handler = TransportHandler;
+	StopTransport();
+	return StartTransport(*Handler);
+}
+
+
+/**
+ * IMessageTransport interface.
+ */
+
+FName FQuicMessageTransport::GetDebugName() const
+{
+	return "QuicMessageTransport";
+}
+
+
+bool FQuicMessageTransport::StartTransport(IMessageTransportHandler& Handler)
+{
+    TransportHandler = &Handler;
+
+	EndpointManager = MakeUnique<FQuicEndpointManager>(EndpointConfig);
+
+	EndpointManager->OnMessageDelivered().BindRaw(this, &FQuicMessageTransport::HandleManagerMessageDelivered);
+	EndpointManager->OnMessageValidated().BindRaw(this, &FQuicMessageTransport::HandleManagerMessageReceived);
+	EndpointManager->OnSerializeHeader().BindRaw(this, &FQuicMessageTransport::HandleManagerSerializeHeader);
+	EndpointManager->OnDeserializeHeader().BindRaw(this, &FQuicMessageTransport::HandleManagerDeserializeHeader);
+	EndpointManager->OnEndpointNodeDiscovered().BindRaw(this, &FQuicMessageTransport::HandleManagerNodeDiscovered);
+	EndpointManager->OnEndpointNodeLost().BindRaw(this, &FQuicMessageTransport::HandleManagerNodeLost);
+
+	if (EndpointMode == EEndpointMode::Server)
+	{
+		EndpointManager->InitializeServer();
+	}
+
+	MessageProcessor = MakeUnique<FQuicMessageProcessor>();
+
+    MessageProcessor->OnMessageDeserialized().BindRaw(this, &FQuicMessageTransport::HandleProcessorMessageDeserialized);
+    MessageProcessor->OnMessageSerialized().BindRaw(this, &FQuicMessageTransport::HandleProcessorMessageSerialized);
+	
+    TSet<FIPv4Endpoint> ProcessStaticEndpints(StaticEndpoints);
+
+    for (const FIPv4Endpoint& StaticEndpoint : ProcessStaticEndpints)
+    {
+		if (StaticEndpoint == FIPv4Endpoint())
+		{
+			continue;
+		}
+
+        AddStaticEndpoint(StaticEndpoint);
+    }
+
+    UE_LOG(LogQuicMessaging, Verbose, TEXT("[MessageTransport] Started transport."));
+
+    return true;
+}
+
+
+void FQuicMessageTransport::StopTransport()
+{
+	UE_LOG(LogQuicMessaging, Verbose, TEXT("[MessageTransport] Stopping transport ..."));
+
+	if (EndpointManager.IsValid())
+	{
+		EndpointManager.Reset();
+		EndpointManager = nullptr;
+	}
+
+    if (MessageProcessor.IsValid())
+    {
+		MessageProcessor.Reset();
+        MessageProcessor = nullptr;
+    }
+
+    TransportHandler = nullptr;
+
+	if (ErrorFuture.IsValid())
+	{
+		ErrorFuture.Reset();
+	}
+
+    UE_LOG(LogQuicMessaging, Verbose, TEXT("[MessageTransport] Stopped transport."));
+}
+
+
+bool FQuicMessageTransport::TransportMessage(
+    const TSharedRef<IMessageContext, ESPMode::ThreadSafe>& Context,
+    const TArray<FGuid>& Recipients)
+{
+    if (MessageProcessor == nullptr)
+    {
+        return false;
+    }
+
+    if (Context->GetRecipients().Num() > QUIC_MESSAGING_MAX_RECIPIENTS)
+    {
+        return false;
+    }
+
+    return MessageProcessor->ProcessOutboundMessage(Context, Recipients, EQuicMessageType::Data);
+}
+
+
+bool FQuicMessageTransport::TransportAuthMessage(
+    const TSharedRef<IMessageContext, ESPMode::ThreadSafe>& Context,
+    const FGuid& Recipient)
+{
+    if (MessageProcessor == nullptr)
+    {
+        return false;
+    }
+
+	const TArray<FGuid> Recipients = { Recipient };
+
+    return MessageProcessor->ProcessOutboundMessage(Context, Recipients, EQuicMessageType::Authentication);
+}
+
+
+bool FQuicMessageTransport::TransportAuthResponseMessage(
+    const TSharedRef<IMessageContext, ESPMode::ThreadSafe>& Context,
+    const FGuid& Recipient)
+{
+    if (MessageProcessor == nullptr)
+    {
+        return false;
+    }
+
+	const TArray<FGuid> Recipients = { Recipient };
+
+    return MessageProcessor->ProcessOutboundMessage(Context, Recipients, EQuicMessageType::AuthenticationResponse);
+}
+
+
+FOnQuicMetaMessageReceived& FQuicMessageTransport::OnMetaMessageReceived()
+{
+	return OnQuicMetaMessageDelegate;
+}
+
+
+void FQuicMessageTransport::HandleProcessorMessageDeserialized(
+	const TSharedPtr<FQuicDeserializedMessage, ESPMode::ThreadSafe> DeserializedMessage,
+	const FGuid NodeId) const
+{
+	if (!DeserializedMessage.IsValid())
+	{
+		UE_LOG(LogQuicMessaging, Error,
+			TEXT("[MessageTransport] Deserialized message pointer invalid."));
+
+		return;
+	}
+
+	// Receive authentication messages separately
+	if (DeserializedMessage->GetMessageTypeInfo()
+		->IsChildOf(FQuicMetaMessage::StaticStruct()))
+	{
+		OnQuicMetaMessageDelegate.Broadcast(NodeId, DeserializedMessage.ToSharedRef());
+
+		return;
+	}
+
+	TransportHandler->ReceiveTransportMessage(DeserializedMessage.ToSharedRef(), NodeId);
+}
+
+
+void FQuicMessageTransport::HandleProcessorMessageSerialized(
+	TSharedPtr<FQuicSerializedMessage, ESPMode::ThreadSafe> SerializedMessage)
+{
+	if (EndpointManager == nullptr)
+	{
+		return;
+	}
+
+	const bool bIsBroadcast = SerializedMessage->GetRecipients().IsEmpty();
+
+	TArray<FGuid> Recipients = bIsBroadcast
+		? EndpointManager->GetKnownNodeIds() : SerializedMessage->GetRecipients();
+
+	TArray<TTuple<FGuid, uint32>> MessageMetas;
+
+	for (const FGuid& Recipient : Recipients)
+	{
+		const uint32 MessageId = SerializedMessageId;
+
+		SerializedMessages.Add(MessageId, SerializedMessage);
+		MessageMetas.Add(TTuple<FGuid, uint32>(Recipient, MessageId));
+
+		SerializedMessageId++;
+	}
+
+	EndpointManager->EnqueueOutboundMessages(
+		SerializedMessage->GetDataPtr(), MessageMetas,
+		SerializedMessage->GetMessageType());
+}
+
+
+void FQuicMessageTransport::HandleManagerMessageDelivered(const uint32 MessageId)
+{
+	if (!SerializedMessages.Contains(MessageId))
+	{
+		return;
+	}
+
+	SerializedMessages.Remove(MessageId);
+}
+
+
+void FQuicMessageTransport::HandleManagerMessageReceived(const FInboundMessage InboundMessage)
+{
+	if (MessageProcessor == nullptr)
+	{
+		return;
+	}
+
+	MessageProcessor->ProcessInboundMessage(InboundMessage);
+}
+
+
+FQuicPayloadPtr FQuicMessageTransport::HandleManagerSerializeHeader(FMessageHeader& MessageHeader)
+{
+	if (MessageProcessor == nullptr)
+	{
+		return MakeShared<TArray<uint8>, ESPMode::ThreadSafe>();
+	}
+
+	return MessageProcessor->SerializeMessageHeader(MessageHeader);
+}
+
+
+FMessageHeader FQuicMessageTransport::HandleManagerDeserializeHeader(FQuicPayloadPtr HeaderData)
+{
+	if (MessageProcessor == nullptr)
+	{
+		return FMessageHeader();
+	}
+
+	return MessageProcessor->DeserializeMessageHeader(HeaderData);
+}
+
+
+void FQuicMessageTransport::HandleManagerNodeDiscovered(const FGuid& DiscoveredNodeId)
+{
+	UE_LOG(LogQuicMessaging, Display,
+		TEXT("[MessageProcessor] Discovered node %s."), *DiscoveredNodeId.ToString());
+
+	TransportHandler->DiscoverTransportNode(DiscoveredNodeId);
+}
+
+
+void FQuicMessageTransport::HandleManagerNodeLost(const FGuid& LostNodeId)
+{
+	UE_LOG(LogQuicMessaging, Display,
+		TEXT("[MessageProcessor] Lost node %s."), *LostNodeId.ToString());
+
+	TransportHandler->ForgetTransportNode(LostNodeId);
+}
+
+
diff --git a/Engine/Plugins/Experimental/QuicMessaging/Source/QuicMessaging/Private/Transport/QuicMessageTransport.h b/Engine/Plugins/Experimental/QuicMessaging/Source/QuicMessaging/Private/Transport/QuicMessageTransport.h
new file mode 100644
--- /dev/null
+++ b/Engine/Plugins/Experimental/QuicMessaging/Source/QuicMessaging/Private/Transport/QuicMessageTransport.h
@@ -0,0 +1,299 @@
+// Copyright Epic Games, Inc. All Rights Reserved.
+
+#pragma once
+
+#include "CoreTypes.h"
+#include "IMessageTransport.h"
+#include "Interfaces/IPv4/IPv4Endpoint.h"
+#include "IMessageAttachment.h"
+#include "Templates/SharedPointer.h"
+#include "Async/Future.h"
+#include "Containers/Ticker.h"
+#include "INetworkMessagingExtension.h"
+
+#include "Serialization/QuicSerializedMessage.h"
+#include "Serialization/QuicDeserializedMessage.h"
+
+#include "QuicFlags.h"
+#include "QuicMessages.h"
+#include "QuicTransportMessages.h"
+
+
+class FArrayReader;
+class FRunnableThread;
+
+class FQuicEndpointManager;
+struct FInboundMessage;
+struct FQuicEndpointConfig;
+
+class FQuicMessageProcessor;
+class IMessageAttachment;
+class IMessageContext;
+class IMessageTransportHandler;
+enum class EQuicMessageFormat : uint8;
+
+struct FGuid;
+
+
+/**
+* Implements a message transport technology using mvfst with QUIC.
+*/
+class FQuicMessageTransport
+	: public TSharedFromThis<FQuicMessageTransport, ESPMode::ThreadSafe>
+	, public IMessageTransport
+{
+public:
+
+	/** Default constructor. */
+	FQuicMessageTransport() = delete;
+
+	/**
+	 * Creates and initializes a new transport instance.
+	 *
+	 * @param bInIsClient Whether the transport is in client mode
+	 * @param InEndpointConfig The endpoint config
+	 * @param InStaticEndpoints The static endpoints
+	 */
+	FQuicMessageTransport(const bool bInIsClient,
+		const TSharedRef<FQuicEndpointConfig> InEndpointConfig,
+		TArray<FIPv4Endpoint> InStaticEndpoints);
+
+	/** Virtual destructor. */
+	virtual ~FQuicMessageTransport();
+
+	/** Notifies the transport that the application is about to exit. */
+	void OnAppPreExit();
+
+public:
+
+	/**
+	 * Add a static server endpoint to the transport.
+	 * 
+	 * @param InEndpoint The static endpoint to add
+	 */
+	void AddStaticEndpoint(const FIPv4Endpoint& InEndpoint);
+
+	/**
+	 * Remove a static server endpoint from the transport.
+	 * 
+	 * @param InEndpoint The static endpoint to remove
+	 */
+	void RemoveStaticEndpoint(const FIPv4Endpoint& InEndpoint);
+
+	/**
+	 * Get the IPv4Endpoint listening addresses.
+	 */
+	TArray<FIPv4Endpoint> GetListeningAddresses() const;
+
+	/**
+	 * Get the known endpoints that are in active communication with the QUIC transport.
+	 * This is information known by the EndpointManager.
+	 */
+	TArray<FIPv4Endpoint> GetKnownEndpoints() const;
+
+	/**
+	 * Get the latest network/transport statistics for a node.
+	 *
+	 * @param NodeId The node ID
+	 */
+	FMessageTransportStatistics GetLatestStatistics(FGuid NodeId) const;
+
+public:
+
+	/**
+	 * Set the maximum authentication message size.
+	 *
+	 * @param MaxBytes The maximum message size in bytes
+	 */
+	void SetMaxAuthenticationMessageSize(const uint32 MaxBytes) const;
+
+	/**
+	 * Get the NodeId of the remote endpoint.
+	 *
+	 * @param RemoteEndpoint The remote endpoint address
+	 *
+	 * @return Optional FGuid if the node was found
+	 */
+	TOptional<FGuid> GetNodeId(const FIPv4Endpoint& RemoteEndpoint) const;
+
+	/**
+	 * Check if a node is authenticated.
+	 *
+	 * @param NodeId The node identifier
+	 */
+	bool IsNodeAuthenticated(const FGuid& NodeId) const;
+
+	/**
+	 * Mark a node as authenticated.
+	 *
+	 * @param NodeId The node identifier
+	 */
+	void SetNodeAuthenticated(const FGuid& NodeId);
+
+	/**
+	 * Disconnect a node.
+	 *
+	 * @param NodeId The node identifier
+	 */
+	void DisconnectNode(const FGuid& NodeId) const;
+
+	/**
+	 * Set the connection cooldown.
+	 *
+	 * @param bEnabled Whether the cooldown should be enabled
+	 * @param MaxAttempts Maximum number of attempts
+	 * @param PeriodSeconds Time period for maximum number of attempts in seconds
+	 * @param CooldownSeconds Duration of the cooldown in seconds
+	 * @param CooldownMaxSeconds Max duration of the cooldown in seconds
+	 */
+	void SetConnectionCooldown(
+		const bool bEnabled, const uint32 MaxAttempts, const uint32 PeriodSeconds,
+		const uint32 CooldownSeconds, const uint32 CooldownMaxSeconds) const;
+
+public:
+
+	/**
+	 * Returns a delegate that is exectued when a connection error happened.
+	 * 
+	 * @return The delegate
+	 * @note this delegate is executed from the main thread by the task graph
+	 */
+	DECLARE_DELEGATE(FOnError)
+	FOnError& OnTransportError()
+	{
+		return TransportErrorDelegate;
+	}
+
+	/**
+	 * Restart the transport using the same Transport handler
+	 * @return true if the restart was successful
+	 */
+	bool RestartTransport();
+
+	//~ IMessageTransport interface
+
+	virtual FName GetDebugName() const override;
+
+	virtual bool StartTransport(IMessageTransportHandler& Handler) override;
+
+	virtual void StopTransport() override;
+
+	virtual bool TransportMessage(
+		const TSharedRef<IMessageContext, ESPMode::ThreadSafe>& Context,
+		const TArray<FGuid>& Recipients) override;
+
+public:
+
+	/**
+	 * Transport an authentication message.
+	 *
+	 * @param Context The message context
+	 * @param Recipient The message recipient
+	 */
+	bool TransportAuthMessage(
+		const TSharedRef<IMessageContext, ESPMode::ThreadSafe>& Context,
+    	const FGuid& Recipient);
+
+	/**
+	 * Transport an authentication response message.
+	 *
+	 * @param Context The message context
+	 * @param Recipient The message recipient
+	 */
+	bool TransportAuthResponseMessage(
+		const TSharedRef<IMessageContext, ESPMode::ThreadSafe>& Context,
+    	const FGuid& Recipient);
+
+	/**
+	 * Delegate that notifies bound functions when
+	 * a meta / authentication message has been received.
+	 *
+	 * @return The delegate
+	 */
+	FOnQuicMetaMessageReceived& OnMetaMessageReceived();
+
+private:
+
+	/**
+	 * Handles deserialized transport messages.
+	 */
+	void HandleProcessorMessageDeserialized(
+		const TSharedPtr<FQuicDeserializedMessage, ESPMode::ThreadSafe> DeserializedMessage,
+		const FGuid NodeId) const;
+
+	/**
+	 * Handles serialized transport messages.
+	 */
+	void HandleProcessorMessageSerialized(
+		TSharedPtr<FQuicSerializedMessage, ESPMode::ThreadSafe> SerializedMessage);
+
+private:
+
+	/**
+	 * Handles delivered messages.
+	 */
+	void HandleManagerMessageDelivered(const uint32 MessageId);
+
+	/**
+	 * Handles received transport messages.
+	 */
+	void HandleManagerMessageReceived(const FInboundMessage InboundMessage);
+
+	/**
+	 * Handles message headers that need to be serialized.
+	 */
+	FQuicPayloadPtr HandleManagerSerializeHeader(FMessageHeader& MessageHeader);
+
+	/**
+	 * Handles message headers that need to be deserialized.
+	 */
+	FMessageHeader HandleManagerDeserializeHeader(FQuicPayloadPtr HeaderData);
+
+	/**
+	 * Handles discovered transport endpoints.
+	 */
+	void HandleManagerNodeDiscovered(const FGuid& DiscoveredNodeId);
+
+	/**
+	 * Handles lost transport endpoints.
+	 */
+	void HandleManagerNodeLost(const FGuid& LostNodeId);
+
+private:
+
+	/** Holds any pending restart requests. */
+	TFuture<void> ErrorFuture;
+
+	/** Holds the transport error delegate. */
+	FOnError TransportErrorDelegate;
+
+	/** Holds the message processor. */
+	TUniquePtr<FQuicMessageProcessor> MessageProcessor = nullptr;
+
+	/** Message transport handler. */
+	IMessageTransportHandler* TransportHandler = nullptr;
+
+	/** Holds the endpoint manager. */
+	TUniquePtr<FQuicEndpointManager> EndpointManager = nullptr;
+
+	/** Flag indicating the endpoint mode. */
+	EEndpointMode EndpointMode = EEndpointMode::Client;
+
+	/** Holds the QuicEndpoint config. */
+	TSharedRef<FQuicEndpointConfig> EndpointConfig;
+
+	/** Holds the static endpoints. */
+	TSet<FIPv4Endpoint> StaticEndpoints;
+
+private:
+
+	/** Current id for serialized messages. */
+	uint32 SerializedMessageId = 0;
+
+	/** Holds serialized messages. */
+	TMap<uint32, TSharedPtr<FQuicSerializedMessage, ESPMode::ThreadSafe>> SerializedMessages;
+
+	/** Holds the meta message received delegate. */
+	FOnQuicMetaMessageReceived OnQuicMetaMessageDelegate;
+
+};
diff --git a/Engine/Plugins/Experimental/QuicMessaging/Source/QuicMessaging/Public/IQuicNetworkMessagingExtension.h b/Engine/Plugins/Experimental/QuicMessaging/Source/QuicMessaging/Public/IQuicNetworkMessagingExtension.h
new file mode 100644
--- /dev/null
+++ b/Engine/Plugins/Experimental/QuicMessaging/Source/QuicMessaging/Public/IQuicNetworkMessagingExtension.h
@@ -0,0 +1,111 @@
+// Copyright Epic Games, Inc. All Rights Reserved.
+
+#pragma once
+
+#include "IMessageContext.h"
+#include "INetworkMessagingExtension.h"
+
+#include "QuicTransportMessages.h"
+
+
+/**
+ * Interface for the messaging module network extension
+ * Plugins or modules implementing messaging transport for MessageBus
+ * can implement this modular feature to provide control on the service it provides.
+ */
+class IQuicNetworkMessagingExtension : public INetworkMessagingExtension
+{
+
+public:
+
+	/**
+	 * Get the generated Guid for this instance.
+	 * 
+	 * @return The Guid
+	 */
+	virtual FGuid GetEndpointGuid() = 0;
+
+	/**
+	 * Set the maximum size an authentication message can have.
+	 *
+	 * @param MaxBytes The maximum message size in bytes
+	 */
+	virtual void SetMaxAuthenticationMessageSize(const uint32 MaxBytes) = 0;
+
+	/**
+	 * Get the NodeId of the remote endpoint.
+	 *
+	 * @param RemoteEndpoint The remote endpoint address:port
+	 *
+	 * @return Optional FGuid if the node was found
+	 */
+	virtual TOptional<FGuid> GetNodeId(const FString& RemoteEndpoint) = 0;
+
+	/**
+	 * Check if this node is authenticated.
+	 * 
+	 * @param NodeId The node identifier
+	 * @return Whether the node is authenticated
+	 */
+	virtual bool IsNodeAuthenticated(const FGuid& NodeId) const = 0;
+
+	/**
+	 * Set this node as authenticated.
+	 * 
+	 * @param NodeId The node identifier
+	 */
+	virtual void SetNodeAuthenticated(const FGuid& NodeId) = 0;
+
+	/**
+	 * Disconnect a node.
+	 *
+	 * @param NodeId The node identifier
+	 */
+	virtual void DisconnectNode(const FGuid& NodeId) = 0;
+
+	/**
+	 * Set the connection cooldown.
+	 *
+	 * @param bEnabled Whether the cooldown should be enabled
+	 * @param MaxAttempts Maximum number of attempts
+	 * @param PeriodSeconds Time period for maximum number of attempts in seconds
+	 * @param CooldownSeconds Duration of the cooldown in seconds
+	 * @param CooldownMaxSeconds Max duration of the cooldown in seconds
+	 */
+	virtual void SetConnectionCooldown(
+		const bool bEnabled, const uint32 MaxAttempts, const uint32 PeriodSeconds,
+		const uint32 CooldownSeconds, const uint32 CooldownMaxSeconds) = 0;
+
+	/**
+	 * Transport authentication message to endpoint.
+	 *
+	 * @param Context The message context
+	 * @param Recipient Guid of message recipient
+	 *
+	 * @return Whether the message was sent
+	 */
+	virtual bool TransportAuthMessage(
+		const TSharedRef<IMessageContext, ESPMode::ThreadSafe>& Context,
+    	const FGuid& Recipient) = 0;
+
+	/**
+	 * Transport authentication response message to endpoint.
+	 *
+	 * @param Context The message context
+	 * @param Recipient Guid of message recipient
+	 *
+	 * @return Whether the message was sent
+	 */
+	virtual bool TransportAuthResponseMessage(
+		const TSharedRef<IMessageContext, ESPMode::ThreadSafe>& Context,
+    	const FGuid& Recipient) = 0;
+
+	/**
+	 * Delegate that notifies bound functions when
+	 * a meta / authentication message has been received.
+	 *
+	 * @return The delegate
+	 */
+	virtual FOnQuicMetaMessageReceived& OnQuicMetaMessageReceived() = 0;
+
+};
diff --git a/Engine/Plugins/Experimental/QuicMessaging/Source/QuicMessaging/Public/QuicTransportMessages.h b/Engine/Plugins/Experimental/QuicMessaging/Source/QuicMessaging/Public/QuicTransportMessages.h
new file mode 100644
--- /dev/null
+++ b/Engine/Plugins/Experimental/QuicMessaging/Source/QuicMessaging/Public/QuicTransportMessages.h
@@ -0,0 +1,116 @@
+// Copyright Epic Games, Inc. All Rights Reserved.
+
+#pragma once
+
+#include "CoreTypes.h"
+#include "IMessageContext.h"
+
+#include "QuicTransportMessages.generated.h"
+
+
+DECLARE_MULTICAST_DELEGATE_TwoParams(FOnQuicMetaMessageReceived,
+	const FGuid /*LocalNodeId*/, const TSharedRef<IMessageContext>& /*Context*/);
+
+
+USTRUCT()
+struct FQuicMetaMessage
+{
+
+	GENERATED_BODY()
+
+    UPROPERTY(VisibleAnywhere, Category = "QuicTransportMessage")
+    FGuid SenderNodeId;
+
+};
+
+
+USTRUCT()
+struct FQuicAuthMessage : public FQuicMetaMessage
+{
+
+	GENERATED_BODY()
+
+	UPROPERTY(VisibleAnywhere, Category = "QuicTransportMessage")
+	FString Payload;
+
+};
+
+
+USTRUCT()
+struct FQuicAuthResponseMessage : public FQuicMetaMessage
+{
+
+	GENERATED_BODY()
+
+	UPROPERTY(VisibleAnywhere, Category = "QuicTransportMessage")
+	bool bAuthSuccessful;
+
+	UPROPERTY(VisibleAnywhere, Category = "QuicTransportMessage")
+	FString Reason;
+
+};
+
+
+class FQuicMetaMessageContext
+	: public IMessageContext
+{
+
+public:
+
+	template <typename MessageType>
+	explicit FQuicMetaMessageContext(MessageType* InMessage)
+		: Message(InMessage)
+		, TypeInfo(InMessage->StaticStruct())
+	{ }
+
+	~FQuicMetaMessageContext()
+	{
+		if (Message)
+		{
+			if (TypeInfo.IsValid())
+			{
+				TypeInfo->DestroyStruct(Message);
+			}
+
+			// @note: Due to the nature of how message contexts work,
+			// there is no other way to do this than to free the message memory.
+			FMemory::Free(Message);
+		}
+	}
+
+public:
+
+	//~ IMessageContext interface
+
+	virtual const TMap<FName, FString>& GetAnnotations() const override { return Annotations; }
+	virtual TSharedPtr<IMessageAttachment, ESPMode::ThreadSafe> GetAttachment() const override { return Attachment; }
+	virtual const FDateTime& GetExpiration() const override { return Expiration; }
+	virtual const void* GetMessage() const override { return Message; }
+	virtual const TWeakObjectPtr<UScriptStruct>& GetMessageTypeInfo() const override { return TypeInfo; }
+	virtual TSharedPtr<IMessageContext, ESPMode::ThreadSafe> GetOriginalContext() const override { return OriginalContext; }
+	virtual const TArray<FMessageAddress>& GetRecipients() const override { return Recipients; }
+	virtual EMessageScope GetScope() const override { return Scope; }
+	virtual EMessageFlags GetFlags() const override { return Flags; }
+	virtual const FMessageAddress& GetSender() const override { return Sender; }
+	virtual const FMessageAddress& GetForwarder() const override { return Sender; }
+	virtual ENamedThreads::Type GetSenderThread() const override { return SenderThread; }
+	virtual const FDateTime& GetTimeForwarded() const override { return TimeSent; }
+	virtual const FDateTime& GetTimeSent() const override { return TimeSent; }
+
+private:
+
+	TMap<FName, FString> Annotations = TMap<FName, FString>();
+	TSharedPtr<IMessageAttachment, ESPMode::ThreadSafe> Attachment = nullptr;
+	FDateTime Expiration = FDateTime::MaxValue();
+	void* Message;
+	TSharedPtr<IMessageContext, ESPMode::ThreadSafe> OriginalContext = nullptr;
+	TArray<FMessageAddress> Recipients = TArray<FMessageAddress>();
+	EMessageScope Scope = EMessageScope::Network;
+	EMessageFlags Flags = EMessageFlags::Reliable;
+	FMessageAddress Sender = FMessageAddress::NewAddress();
+	ENamedThreads::Type SenderThread = ENamedThreads::AnyThread;
+	FDateTime TimeSent = FDateTime::Now();
+	TWeakObjectPtr<UScriptStruct> TypeInfo;
+
+};
+
diff --git a/Engine/Plugins/Experimental/QuicMessaging/Source/QuicMessaging/Public/Shared/QuicMessagingSettings.h b/Engine/Plugins/Experimental/QuicMessaging/Source/QuicMessaging/Public/Shared/QuicMessagingSettings.h
new file mode 100644
--- /dev/null
+++ b/Engine/Plugins/Experimental/QuicMessaging/Source/QuicMessaging/Public/Shared/QuicMessagingSettings.h
@@ -0,0 +1,186 @@
+// Copyright Epic Games, Inc. All Rights Reserved.
+
+#pragma once
+
+#include "CoreTypes.h"
+#include "Containers/Array.h"
+#include "Containers/UnrealString.h"
+#include "UObject/ObjectMacros.h"
+#include "UObject/Object.h"
+
+#include "QuicMessagingSettings.generated.h"
+
+/** Defines the Quic message format available (how the message data is encoded). */
+UENUM()
+enum class EQuicMessageFormat : uint8
+{
+	/** No format specified. Legacy - Not exposed to user.*/
+	None = 0 UMETA(Hidden),
+
+	/** JSON format specified. Legacy - Not exposed to user.*/
+	Json UMETA(Hidden),
+
+	/** Tagged property format specified. Legacy - Not exposed to user.*/
+	TaggedProperty UMETA(Hidden),
+
+	/**
+	 * Quic messages are encoded in CBOR, using the platform endianness. This is the fastest and preferred option, but the CBOR data will not be readable by an external standard-compliant CBOR parser
+	 * if generated from a little endian platform. If the data needs to be consumed outside the Unreal Engine, consider using CborStandardEndianness format instead.
+	 */
+	CborPlatformEndianness UMETA(DisplayName="CBOR (Platform Endianness)"),
+
+	/**
+	 * Quic messages are encoded in CBOR, using the CBOR standard-complinant endianness (big endian). It will perform slower on a little-endian platform, but the data will be readable by standard CBOR parsers.
+	 * Useful if the Quic messages needs to be analyzed/consumed outside the Unreal Engine.
+	 */
+	CborStandardEndianness UMETA(DisplayName="CBOR (Standard Endianness)"),
+};
+
+UCLASS(config=Engine)
+class QUICMESSAGING_API UQuicMessagingSettings
+	: public UObject
+{
+	GENERATED_UCLASS_BODY()
+
+public:
+
+	/**
+	 * Whether Quic messaging is enabled by default. If false -messaging will need to be added
+	 * to the commandline when running non-editor builds.
+	 *
+	 * (Note - in Shipping builds ALLOW_QUIC_MESSAGING_SHIPPING=1 must also be defined in TargetRules
+	 * for messaging to be available with or without this setting)
+	 */
+	UPROPERTY(config, EditAnywhere, Category = Availability)
+	bool EnabledByDefault = false;
+
+	/**
+	 * Whether the Quic transport channel is enabled.
+	 * Can be specified on the command line with `-QUICMESSAGING_TRANSPORT_ENABLE=`
+	 */
+	UPROPERTY(config, EditAnywhere, Category=Transport)
+	bool EnableTransport = true;
+
+	/** Whether the Quic transport channel should try to auto repair when in error. */
+	UPROPERTY(config, EditAnywhere, Category=Transport, AdvancedDisplay)
+	bool bAutoRepair = true;
+
+	/**
+	 * Whether the QUIC transport endpoint is a client (true) or a server (false)
+	 */
+	UPROPERTY(config)
+	bool bIsClient = true;
+
+	/**
+	 * Whether encryption should be used after the quic handshake
+	 */
+	UPROPERTY(config, EditAnywhere, Category = Transport, AdvancedDisplay)
+	bool bEncryption = true;
+
+	/**
+	 * Timeout in seconds when the remote endpoint cannot be discovered.
+	 */
+	UPROPERTY(config, EditAnywhere, Category = Transport, AdvancedDisplay)
+	uint32 DiscoveryTimeoutSeconds = 10;
+
+	/**
+	 * Whether server side authentication is enabled
+	 */
+	UPROPERTY(config, EditAnywhere, Category = Transport, AdvancedDisplay)
+	bool bAuthEnabled = false;
+
+	/**
+	 * Maximum authentication message size for server side authentication
+	 */
+	UPROPERTY(config, EditAnywhere, Category = Transport, AdvancedDisplay)
+	uint32 MaxAuthenticationMessageSize = 8000;
+
+	/**
+	 * Whether server has a cooldown for connection attempts.
+	 */
+	UPROPERTY(config, EditAnywhere, Category = Transport, AdvancedDisplay)
+	bool bConnectionCooldownEnabled = true;
+
+	/**
+	 * Maximum number of connection attempts until the cooldown is triggered.
+	 */
+	UPROPERTY(config, EditAnywhere, Category = Transport, AdvancedDisplay)
+	uint32 ConnectionCooldownMaxAttempts = 5;
+
+	/**
+	 * Time period in seconds within which the maximum attempts must happen.
+	 */
+	UPROPERTY(config, EditAnywhere, Category = Transport, AdvancedDisplay)
+	uint32 ConnectionCooldownPeriodSeconds = 30;
+
+	/**
+	 * Connection cooldown in seconds used with exponential backoff.
+	 */
+	UPROPERTY(config, EditAnywhere, Category = Transport, AdvancedDisplay)
+	uint32 ConnectionCooldownSeconds = 60;
+
+	/**
+	 * Maximum connection cooldown in seconds for exponential backoff.
+	 */
+	UPROPERTY(config, EditAnywhere, Category = Transport, AdvancedDisplay)
+	uint32 ConnectionCooldownMaxSeconds = 3600;
+
+	/**
+	 * Whether the QUIC client should verify the server certificate
+	 */
+	UPROPERTY(config, EditAnywhere, Category = Transport, AdvancedDisplay)
+	bool bClientVerificationEnabled = true;
+
+	/**
+	 * Server certificate to encrypt QUIC transport
+	 */
+	UPROPERTY(config)
+	FString QuicServerCertificate = "";
+
+	/**
+	 * Server private key to encrypt QUIC transport
+	 */
+	UPROPERTY(config)
+	FString QuicServerPrivateKey = "";
+
+	/**
+	 * Guid of this endpoint
+	 */
+	UPROPERTY(config, EditAnywhere, Category = Transport, AdvancedDisplay)
+	FGuid EndpointGuid;
+
+	/** The number of consecutive attempt the auto repair routine will try to repair. */
+	UPROPERTY(config, EditAnywhere, Category = Transport, AdvancedDisplay)
+	uint32 AutoRepairAttemptLimit = 10;
+
+	/** Whether to stop the transport service when the application deactivates, and restart it when the application is reactivated */
+	UPROPERTY(config)
+	bool bStopServiceWhenAppDeactivates = true;
+
+	/**
+	 * The IP endpoint to listen to and send packets from.
+	 *
+	 * The format is IP_ADDRESS:PORT_NUMBER.
+	 * 0.0.0.0:0 will bind to the default network adapter on Windows,
+	 * and all available network adapters on other operating systems.
+	 * Specifying an interface IP here, will use that interface for multicasting and static endpoint *might* also reach this client through <unicast ip:multicast port>
+	 * Specifying both the IP and Port will allow usage of static endpoint to reach this client
+	 * Can be specified on the command line with `-QUICMESSAGING_TRANSPORT_UNICAST=`
+	 */
+	UPROPERTY(config, EditAnywhere, Category=Transport)
+	FString UnicastEndpoint;
+
+	/** The format used to serialize the Quic message payload. */
+	UPROPERTY(config, EditAnywhere, Category=Transport)
+	EQuicMessageFormat MessageFormat = EQuicMessageFormat::CborPlatformEndianness;
+
+	/**
+	 * The IP endpoints of static devices.
+	 *
+	 * Use this setting to reach devices on other subnets, such as mobile phones on a WiFi network.
+	 * The format is IP_ADDRESS:PORT_NUMBER.
+	 */
+	UPROPERTY(config, EditAnywhere, Category=Transport, AdvancedDisplay)
+	TArray<FString> StaticEndpoints;
+
+};
diff --git a/Engine/Plugins/Experimental/QuicMessaging/Source/QuicMessaging/QuicMessaging.Build.cs b/Engine/Plugins/Experimental/QuicMessaging/Source/QuicMessaging/QuicMessaging.Build.cs
new file mode 100644
--- /dev/null
+++ b/Engine/Plugins/Experimental/QuicMessaging/Source/QuicMessaging/QuicMessaging.Build.cs
@@ -0,0 +1,71 @@
+// Copyright Epic Games, Inc. All Rights Reserved.
+using System.IO;
+
+namespace UnrealBuildTool.Rules
+{
+	public class QuicMessaging : ModuleRules
+	{
+		public QuicMessaging(ReadOnlyTargetRules Target) : base(Target)
+		{
+			PublicDependencyModuleNames.AddRange(
+				new string[] {
+					"Core",
+					"QuicMessagingTransport"
+				}
+			);
+
+			PrivateDependencyModuleNames.AddRange(
+				new string[] {
+					"CoreUObject",
+					"Json",
+					"JsonUtilities",
+					"Cbor",
+					"Messaging",
+					"Networking",
+					"Serialization",
+					"TraceLog",
+					"Sockets",
+					"Projects",
+					"MsQuicRuntime"
+				}
+			);
+			
+			PrivateIncludePathModuleNames.AddRange(
+				new string[] {
+					"MessagingCommon",
+					"MsQuicRuntime"
+				}
+			);
+
+			PrivateIncludePaths.AddRange(
+				new string[] {
+					"QuicMessaging/Private",
+					"QuicMessaging/Private/Transport",
+					"QuicMessaging/Private/Serialization",
+					"QuicMessaging/Public/Shared",
+					"QuicMessaging/Public",
+				}
+			);
+
+			DynamicallyLoadedModuleNames.AddRange(
+				new string[] {
+				}
+			);
+
+			if (Target.Type == TargetType.Editor) {
+
+				DynamicallyLoadedModuleNames.AddRange(
+					new string[] {
+						"Settings",
+					}
+				);
+
+				PrivateIncludePathModuleNames.AddRange(
+					new string[] {
+						"Settings",
+					}
+				);
+			}
+		}
+	}
+}
diff --git a/Engine/Plugins/Experimental/QuicMessaging/Source/QuicMessagingTransport/Private/QuicBuffers.h b/Engine/Plugins/Experimental/QuicMessaging/Source/QuicMessagingTransport/Private/QuicBuffers.h
new file mode 100644
--- /dev/null
+++ b/Engine/Plugins/Experimental/QuicMessaging/Source/QuicMessagingTransport/Private/QuicBuffers.h
@@ -0,0 +1,92 @@
+// Copyright Epic Games, Inc. All Rights Reserved.
+
+#pragma once
+
+#include "CoreTypes.h"
+
+
+struct FMessageHeader;
+
+using FQuicPayloadPtr = TSharedPtr<TArray<uint8>, ESPMode::ThreadSafe>;
+
+
+struct FInboundBuffer
+{
+	/** Flag indicating whether the header was deserialized. */
+	bool bHeaderDeserialized;
+
+	/** Holds the Quic handler address. */
+	QUIC_ADDR HandlerAddress;
+
+	/** Holds the Quic handler endpoint. */
+	FIPv4Endpoint HandlerEndpoint;
+
+	/** Holds the message header. */
+	FMessageHeader MessageHeader;
+
+	/** Holds the inbound data. */
+	FQuicPayloadPtr InboundData;
+
+	/** Default constructor. */
+	FInboundBuffer()
+		: bHeaderDeserialized(false)
+		, HandlerAddress(QUIC_ADDR())
+		, HandlerEndpoint(FIPv4Endpoint::Any)
+		, MessageHeader(FMessageHeader())
+		, InboundData(nullptr)
+	{ }
+
+	/** Creates and initializes a new instance. */
+	FInboundBuffer(const QUIC_ADDR InHandlerAddress,
+		const FIPv4Endpoint InHandlerEndpoint, const FQuicPayloadPtr InInboundData)
+		: bHeaderDeserialized(false)
+		, HandlerAddress(InHandlerAddress)
+		, HandlerEndpoint(InHandlerEndpoint)
+		, MessageHeader(FMessageHeader())
+		, InboundData(InInboundData)
+	{ }
+};
+
+struct FInboundQuicBuffer
+{
+	/** Holds the Quic stream ID. */
+	uint64 StreamId;
+
+	/** Holds the Quic stream handle. */
+	HQUIC Stream;
+
+	/** Flag indicating whether this is the end of the Quic stream. */
+	bool bEndOfStream;
+
+	/** Holds the peer address. */
+	FIPv4Endpoint PeerAddress;
+
+	/** Holds the total length of all buffers. */
+	uint64 BufferLength;
+
+	/** Holds buffer pointers and buffer sizes. */
+	TMap<uint8*, uint32> Data;
+
+	/** Default constructor. */
+	FInboundQuicBuffer()
+		: StreamId(0)
+		, Stream(nullptr)
+		, bEndOfStream(false)
+		, PeerAddress(FIPv4Endpoint::Any)
+		, BufferLength(0)
+		, Data(TMap<uint8*, uint32>())
+	{ }
+
+	/** Creates and initializes a new instance. */
+	FInboundQuicBuffer(const uint64 InStreamId, HQUIC InStream,
+		const bool bInEndOfStream, const FIPv4Endpoint InPeerAddress,
+		const uint64 InBufferLength)
+		: StreamId(InStreamId)
+		, Stream(InStream)
+		, bEndOfStream(bInEndOfStream)
+		, PeerAddress(InPeerAddress)
+		, BufferLength(InBufferLength)
+		, Data(TMap<uint8*, uint32>())
+	{ }
+};
+
diff --git a/Engine/Plugins/Experimental/QuicMessaging/Source/QuicMessagingTransport/Private/QuicCertificate.h b/Engine/Plugins/Experimental/QuicMessaging/Source/QuicMessagingTransport/Private/QuicCertificate.h
new file mode 100644
--- /dev/null
+++ b/Engine/Plugins/Experimental/QuicMessaging/Source/QuicMessagingTransport/Private/QuicCertificate.h
@@ -0,0 +1,214 @@
+// Copyright Epic Games, Inc. All Rights Reserved.
+
+#pragma once
+
+#include "QuicIncludes.h"
+
+#include "HAL/FileManagerGeneric.h"
+#include "Misc/FileHelper.h"
+#include "Misc/Paths.h"
+
+
+namespace QuicCertificateUtils
+{
+
+	/**
+	 * Helper to setup a QUIC credential config.
+	 */
+	typedef struct QUIC_CREDENTIAL_CONFIG_HELPER
+	{
+		QUIC_CREDENTIAL_CONFIG CredConfig;
+
+		QUIC_CERTIFICATE_FILE CertificateFile;
+
+	} QUIC_CREDENTIAL_CONFIG_HELPER;
+
+
+	/**
+	 * Generate a 2048-bit RSA key.
+	 *
+	 * @param PrivateKey The private key struct
+	 */
+	bool GeneratePrivateKey(EVP_PKEY* PrivateKey)
+	{
+		if (!PrivateKey)
+		{
+			return false;
+		}
+
+		/* Generate the RSA key and assign it to pkey. */
+		RSA* Rsa = RSA_new();
+
+		{
+			BIGNUM* BigNum = BN_new();
+			BN_set_word(BigNum, RSA_F4);
+			RSA_generate_key_ex(Rsa, 2048, BigNum, nullptr);
+			BN_free(BigNum);
+		}
+
+		if (!EVP_PKEY_assign_RSA(PrivateKey, Rsa))
+		{
+			EVP_PKEY_free(PrivateKey);
+			return false;
+		}
+
+		return true;
+	}
+
+
+	/**
+	 * Converts a private key to a string.
+	 *
+	 * @param PrivateKey The private key struct
+	 */
+	char* PrivateKeyToString(EVP_PKEY* PrivateKey)
+	{
+		BIO* Bio = BIO_new(BIO_s_mem());
+
+		int Ret = PEM_write_bio_PrivateKey(Bio, PrivateKey, NULL, NULL, 0, NULL, NULL);
+
+		if (Ret != 1)
+		{
+			BIO_free(Bio);
+		}
+
+		BIO_flush(Bio);
+
+		char* Text;
+		long Length = BIO_get_mem_data(Bio, &Text);
+		char* Result = new char[Length + 1];
+		FMemory::Memcpy(Result, Text, Length);
+		Result[Length] = 0;
+
+		BIO_free(Bio); // text is just a pointer to bio's data, no need to free that as well
+
+		return Result;
+	}
+
+
+	/**
+	 * Generate self-signed x509 certificate.
+	 *
+	 * @param PrivateKey The private key struct
+	 * @param x509 The certificate struct
+	 */
+	bool GenerateCertificate(EVP_PKEY* PrivateKey, X509* x509)
+	{
+		if (!x509)
+		{
+			return false;
+		}
+
+		/* Set the serial number. */
+		ASN1_INTEGER_set(X509_get_serialNumber(x509), 1);
+
+		/* This certificate is valid from now until exactly one year from now. */
+		X509_gmtime_adj(X509_get_notBefore(x509), 0);
+		X509_gmtime_adj(X509_get_notAfter(x509), 31536000L);
+
+		/* Set the public key for our certificate. */
+		X509_set_pubkey(x509, PrivateKey);
+
+		/* We want to copy the subject name to the issuer name. */
+		X509_NAME* SubjectName = X509_get_subject_name(x509);
+
+		/* Set the country code and common name. */
+		X509_NAME_add_entry_by_txt(SubjectName, "C", MBSTRING_ASC, (unsigned char*)"CA", -1, -1, 0);
+		X509_NAME_add_entry_by_txt(SubjectName, "O", MBSTRING_ASC, (unsigned char*)"MyCompany", -1, -1, 0);
+		X509_NAME_add_entry_by_txt(SubjectName, "CN", MBSTRING_ASC, (unsigned char*)"localhost", -1, -1, 0);
+
+		/* Now set the issuer name. */
+		X509_set_issuer_name(x509, SubjectName);
+
+		/* Actually sign the certificate with our key. */
+		if (!X509_sign(x509, PrivateKey, EVP_sha1()))
+		{
+			//std::cerr << "Error signing certificate." << std::endl;
+			X509_free(x509);
+			return false;
+		}
+
+		return true;
+	}
+
+
+	/**
+	 * Convert a certificate to a string.
+	 *
+	 * @param x509 The certificate struct
+	 */
+	char* X509ToString(X509* x509)
+	{
+		BIO* Bio = BIO_new(BIO_s_mem());
+		int Ret = PEM_write_bio_X509(Bio, x509);
+
+		char* Text;
+		long Length = BIO_get_mem_data(Bio, &Text);
+		char* Result = new char[Length + 1];
+		FMemory::Memcpy(Result, Text, Length);
+		Result[Length] = 0;
+
+		BIO_free(Bio); // text is just a pointer to bio's data, no need to free that as well
+		return Result;
+	}
+
+
+	/**
+	 * Get the paths for self-signed certificate and key.
+	 */
+	TPair<FString, FString> GetSelfSignedPaths()
+	{
+		IPlatformFile& FileManager = FPlatformFileManager::Get().GetPlatformFile();
+
+		FString EngineCertificates = FileManager.ConvertToAbsolutePathForExternalAppForWrite(
+			*FPaths::EngineContentDir()) + TEXT("Certificates");
+
+		FString CertificatesDir = EngineCertificates + TEXT("/QuicMessaging");
+
+		FString CertificateFile = CertificatesDir + TEXT("/transportcert.pem");
+		FString PrivateKeyFile = CertificatesDir + TEXT("/transportkey.key");
+
+		return TPair<FString, FString>(CertificateFile, PrivateKeyFile);
+	}
+
+
+	/**
+	 * Generates a self-signed certificate and returns paths to (certificate, private key).
+	 */
+	bool CreateSelfSigned()
+	{
+		/* Allocate memory for the EVP_PKEY structure. */
+		EVP_PKEY* PrivateKey = EVP_PKEY_new();
+
+		if (!GeneratePrivateKey(PrivateKey))
+		{
+			return false;
+		}
+
+		/* Allocate memory for the X509 structure. */
+		X509* x509 = X509_new();
+
+		if (!GenerateCertificate(PrivateKey, x509))
+		{
+			return false;
+		}
+
+		TPair<FString, FString> Paths = GetSelfSignedPaths();
+
+		FString CertificateFile = Paths.Key;
+		FString PrivateKeyFile = Paths.Value;
+
+		FString CertificateString = ANSI_TO_TCHAR(X509ToString(x509));
+		FString PrivateKeyString = CertificateString + ANSI_TO_TCHAR(PrivateKeyToString(PrivateKey));
+
+		FFileHelper::SaveStringToFile(CertificateString, *CertificateFile);
+		FFileHelper::SaveStringToFile(PrivateKeyString, *PrivateKeyFile);
+
+		X509_free(x509);
+		EVP_PKEY_free(PrivateKey);
+
+		return true;
+	}
+
+};
+
diff --git a/Engine/Plugins/Experimental/QuicMessaging/Source/QuicMessagingTransport/Private/QuicClient.cpp b/Engine/Plugins/Experimental/QuicMessaging/Source/QuicMessagingTransport/Private/QuicClient.cpp
new file mode 100644
--- /dev/null
+++ b/Engine/Plugins/Experimental/QuicMessaging/Source/QuicMessagingTransport/Private/QuicClient.cpp
@@ -0,0 +1,707 @@
+// Copyright Epic Games, Inc. All Rights Reserved.
+
+#include "QuicClient.h"
+
+
+void FQuicClient::StartEndpoint()
+{
+	QUIC_STATUS Status;
+
+	/**
+	 * Allocate new connection object
+	 */
+	if (QUIC_FAILED(Status = MsQuic->ConnectionOpen(Registration,
+		ClientConnectionCallback, this, &ClientConnection)))
+	{
+		ConnectionState = EQuicClientState::Failed;
+		EndpointError = QuicUtils::ConvertQuicStatus(Status);
+
+		UE_LOG(LogQuicMessagingTransport, Error,
+			TEXT("[QuicClient] Could not open connection: %s."),
+			*QuicUtils::GetEndpointErrorString(EndpointError));
+
+		return;
+	}
+
+	if (ClientConfig->EncryptionMode == EEncryptionMode::Disabled)
+	{
+		constexpr BOOLEAN Value = true;
+
+		MsQuic->SetParam(ClientConnection,
+			QUIC_PARAM_CONN_DISABLE_1RTT_ENCRYPTION,
+			sizeof(Value), &Value);
+	}
+
+	/**
+	 * Start connection to the server
+	 */
+	if (QUIC_FAILED(Status = MsQuic->ConnectionStart(ClientConnection,
+		QuicConfig, QUIC_ADDRESS_FAMILY_INET,
+		StringCast<ANSICHAR>(*Endpoint.Address.ToString()).Get(), Endpoint.Port)))
+	{
+		ConnectionState = EQuicClientState::Failed;
+		EndpointError = QuicUtils::ConvertQuicStatus(Status);
+
+		UE_LOG(LogQuicMessagingTransport, Error,
+			TEXT("[QuicClient] Could not start connection: %s."),
+			*QuicUtils::GetEndpointErrorString(EndpointError));
+
+		return;
+	}
+
+	UE_LOG(LogQuicMessagingTransport, Verbose, TEXT("[QuicClient] Started client."));
+}
+
+
+void FQuicClient::StopEndpoint()
+{
+	if (IsStopping())
+	{
+		return;
+	}
+
+	UE_LOG(LogQuicMessagingTransport, Verbose, TEXT("[QuicClient] Stopping client ..."));
+
+	Stop();
+
+	if (ConnectionState == EQuicClientState::Connecting)
+	{
+		ConnectionState = EQuicClientState::Disconnected;
+	}
+
+	if (MsQuic)
+	{
+		if (ClientConnection)
+		{
+			// Only do this if the connection wasn't already shutdown/closed by the peer or transport
+			if (ConnectionState != EQuicClientState::Stopped)
+			{
+				if ((EndpointError == EQuicEndpointError::Normal
+					|| EndpointError == EQuicEndpointError::ConnectionAbort)
+					&& ConnectionState != EQuicClientState::Closed)
+				{
+					MsQuic->ConnectionShutdown(ClientConnection, QUIC_CONNECTION_SHUTDOWN_FLAG_NONE, 0);
+					MsQuic->ConnectionClose(ClientConnection);
+				}
+
+				ConnectionState = EQuicClientState::Stopped;
+			}
+		}
+	}
+
+	UE_LOG(LogQuicMessagingTransport, Verbose, TEXT("[QuicClient] Stopped client."));
+}
+
+
+bool FQuicClient::IsEndpointReady()
+{
+	return ConnectionState == EQuicClientState::Connected;
+}
+
+
+void FQuicClient::UpdateEndpoint()
+{
+	if (IsStopping() || ConnectionState != EQuicClientState::Connected)
+	{
+		return;
+	}
+
+	ConsumeMessages();
+}
+
+
+void FQuicClient::ConsumeInboundMessages()
+{
+	if (InboundMessages.IsEmpty())
+	{
+		return;
+	}
+
+	double CurrentTime = FPlatformTime::Seconds();
+	const double ConsumeLimit = CurrentTime + MaxConsumeTime.GetTotalSeconds();
+
+	FInboundMessage InboundMessage;
+
+	while (CurrentTime < ConsumeLimit && InboundMessages.Dequeue(InboundMessage))
+	{
+		MessageInboundDelegate.ExecuteIfBound(InboundMessage);
+
+		CurrentTime = FPlatformTime::Seconds();
+	}
+}
+
+
+void FQuicClient::ProcessInboundBuffers()
+{
+	if (InboundQuicBuffers.IsEmpty())
+	{
+		return;
+	}
+
+	double CurrentTime = FPlatformTime::Seconds();
+	const double ConsumeLimit = CurrentTime + MaxConsumeTime.GetTotalSeconds();
+
+	FInboundQuicBuffer InboundBuffer;
+
+	while (CurrentTime < ConsumeLimit && InboundQuicBuffers.Dequeue(InboundBuffer))
+	{
+		if (ReceiveFromStream(InboundBuffer))
+		{
+			if (InboundBuffer.bEndOfStream)
+			{
+				ReachedEndOfStream(InboundBuffer);
+			}
+		}
+
+		CurrentTime = FPlatformTime::Seconds();
+	}
+}
+
+
+void FQuicClient::ConsumeOutboundMessages()
+{
+	if (OutboundMessages.IsEmpty())
+	{
+		return;
+	}
+
+	double CurrentTime = FPlatformTime::Seconds();
+	const double ConsumeLimit = CurrentTime + MaxConsumeTime.GetTotalSeconds();
+
+	FOutboundMessage OutboundMessage;
+
+	while (CurrentTime < ConsumeLimit && OutboundMessages.Dequeue(OutboundMessage))
+	{
+		SendToStream(OutboundMessage);
+
+		CurrentTime = FPlatformTime::Seconds();
+	}
+}
+
+
+void FQuicClient::CollectStatistics()
+{
+	if (IsStopping() || ConnectionState != EQuicClientState::Connected)
+	{
+		return;
+	}
+
+	if (MsQuic == nullptr)
+	{
+		return;
+	}
+
+	QUIC_STATISTICS Stats{};
+	
+	if (const QUIC_STATUS Status
+		= GetConnectionStats(ClientConnection, Stats) != QUIC_STATUS_SUCCESS)
+	{
+		UE_LOG(LogQuicMessagingTransport, Error,
+			TEXT("[QuicClient] Could not collect statistics: %s."),
+			*QuicUtils::ConvertResult(Status));
+
+		return;
+	}
+
+	FMessageTransportStatistics TransportStats = ConvertStatistics(Endpoint, Stats);
+
+	StatisticsUpdatedDelegate.ExecuteIfBound(TransportStats);
+}
+
+
+void FQuicClient::EnqueueOutboundMessage(const FOutboundMessage& OutboundMessage)
+{
+	OutboundMessages.Enqueue(OutboundMessage);
+}
+
+
+void FQuicClient::SendHello(FIPv4Endpoint RemoteEndpoint)
+{
+	if (IsStopping() || ConnectionState != EQuicClientState::Connected)
+	{
+		return;
+	}
+
+	FMessageHeader MessageHeader(EQuicMessageType::Hello, FGuid(), LocalNodeId, 0);
+	const FQuicPayloadPtr MetaBuffer = SerializeHeaderDelegate.Execute(MessageHeader);
+
+	const FOutboundMessage OutboundMessage(RemoteEndpoint, MetaBuffer);
+
+	EnqueueOutboundMessage(OutboundMessage);
+
+	UE_LOG(LogQuicMessagingTransport, Verbose, TEXT("[QuicClient] Sent hello."));
+}
+
+
+void FQuicClient::SendBye()
+{
+	if (IsStopping() || ConnectionState != EQuicClientState::Connected)
+	{
+		return;
+	}
+
+	FMessageHeader MessageHeader(EQuicMessageType::Bye, FGuid(), LocalNodeId, 0);
+	const FQuicPayloadPtr MetaBuffer = SerializeHeaderDelegate.Execute(MessageHeader);
+	
+	const FOutboundMessage OutboundMessage(Endpoint, MetaBuffer);
+
+	SendToStream(OutboundMessage);
+
+	UE_LOG(LogQuicMessagingTransport, Verbose, TEXT("[QuicClient] Sent bye."));
+}
+
+
+QUIC_STATUS FQuicClient::ClientConnectionCallback(HQUIC Connection,
+	void* Context, QUIC_CONNECTION_EVENT* Event)
+{
+	FQuicClient* ClientContext = static_cast<FQuicClient*>(Context);
+
+	switch (Event->Type)
+	{
+
+	case QUIC_CONNECTION_EVENT_CONNECTED:
+
+		/**
+		 * Connection handshake has completed
+		 */
+
+		ClientContext->OnConnectionConnected();
+
+		break;
+
+	case QUIC_CONNECTION_EVENT_SHUTDOWN_INITIATED_BY_TRANSPORT:
+
+		/**
+		 * Idle timeout killed the connection
+		 */
+
+		ClientContext->OnConnectionShutdownByTransport(Event->SHUTDOWN_INITIATED_BY_TRANSPORT.Status);
+
+		break;
+
+	case QUIC_CONNECTION_EVENT_SHUTDOWN_INITIATED_BY_PEER:
+
+		/**
+		 * Peer killed the connection
+		 */
+
+		ClientContext->OnConnectionShutdownByPeer(Event->SHUTDOWN_INITIATED_BY_PEER.ErrorCode);
+
+		break;
+
+	case QUIC_CONNECTION_EVENT_SHUTDOWN_COMPLETE:
+
+		/**
+		 * Connection shutdown complete, ready to cleanup
+		 */
+
+		ClientContext->OnConnectionShutdownComplete(Connection, Event->SHUTDOWN_COMPLETE.AppCloseInProgress);
+
+		break;
+
+	case QUIC_CONNECTION_EVENT_RESUMPTION_TICKET_RECEIVED:
+
+		/**
+		 * Resumption ticket received. Does not do anything at this time
+		 */
+
+		ClientContext->OnConnectionResumptionTicketReceived();
+
+		break;
+
+	case QUIC_CONNECTION_EVENT_PEER_STREAM_STARTED:
+	{
+		/**
+		 * Peer started a stream, set local stream callback handler
+		 */
+
+		FIPv4Endpoint* PeerAddress = new FIPv4Endpoint(ClientContext->Endpoint);
+
+		FStreamCallbackInfo* StreamCallbackInfo = new FStreamCallbackInfo();
+		StreamCallbackInfo->Endpoint = ClientContext;
+		StreamCallbackInfo->PeerAddress = PeerAddress;
+
+		ClientContext->MsQuic->SetCallbackHandler(Event->PEER_STREAM_STARTED.Stream,
+			reinterpret_cast<void*>(ClientStreamCallback), StreamCallbackInfo);
+
+		ClientContext->OnConnectionPeerStreamStarted(Event->PEER_STREAM_STARTED.Stream);
+
+		break;
+	}
+
+	case QUIC_CONNECTION_EVENT_DATAGRAM_RECEIVED:
+	case QUIC_CONNECTION_EVENT_DATAGRAM_SEND_STATE_CHANGED:
+	case QUIC_CONNECTION_EVENT_DATAGRAM_STATE_CHANGED:
+	case QUIC_CONNECTION_EVENT_IDEAL_PROCESSOR_CHANGED:
+	case QUIC_CONNECTION_EVENT_LOCAL_ADDRESS_CHANGED:
+	case QUIC_CONNECTION_EVENT_PEER_ADDRESS_CHANGED:
+	case QUIC_CONNECTION_EVENT_PEER_CERTIFICATE_RECEIVED:
+	case QUIC_CONNECTION_EVENT_PEER_NEEDS_STREAMS:
+	case QUIC_CONNECTION_EVENT_RESUMED:
+	case QUIC_CONNECTION_EVENT_STREAMS_AVAILABLE:
+
+		break;
+
+	default:
+
+		UE_LOG(LogQuicMessagingTransport, Warning,
+			TEXT("[QuicClient] Received unhandled connection event: %d"), Event->Type);
+
+		break;
+
+	}
+
+	return QUIC_STATUS_SUCCESS;
+}
+
+
+QUIC_STATUS FQuicClient::ClientStreamCallback(HQUIC Stream,
+	void* Context, QUIC_STREAM_EVENT* Event)
+{
+	const FStreamCallbackInfo* CallbackInfo = static_cast<FStreamCallbackInfo*>(Context);
+	FQuicClient* ClientContext = static_cast<FQuicClient*>(CallbackInfo->Endpoint);
+
+	switch (Event->Type)
+	{
+
+	case QUIC_STREAM_EVENT_SEND_COMPLETE:
+
+		/**
+		 * Previous StreamSend has completed, context returned
+		 */
+
+		ClientContext->OnStreamSendComplete(Stream, Event);
+
+		break;
+
+	case QUIC_STREAM_EVENT_RECEIVE:
+	
+		/**
+		 * Data received from the peer
+		 */
+
+		ClientContext->OnStreamReceive(Stream, *CallbackInfo->PeerAddress, Event);
+
+		return QUIC_STATUS_PENDING;
+	
+
+	case QUIC_STREAM_EVENT_PEER_SEND_ABORTED:
+
+		/**
+		 * Peer aborted its stream send direction
+		 */
+
+		ClientContext->OnStreamPeerSendAborted(Stream);
+
+		break;
+
+	case QUIC_STREAM_EVENT_PEER_SEND_SHUTDOWN:
+
+		/**
+		 * Peer gracefully shut down its stream send direction
+		 */
+
+		break;
+
+	case QUIC_STREAM_EVENT_SHUTDOWN_COMPLETE:
+	
+		/**
+		 * Both stream directions have been shut down, safe to cleanup
+		 */
+
+		ClientContext->OnStreamShutdownComplete(Stream);
+		delete CallbackInfo;
+
+		break;
+
+	case QUIC_STREAM_EVENT_START_COMPLETE:
+	case QUIC_STREAM_EVENT_PEER_ACCEPTED:
+	case QUIC_STREAM_EVENT_PEER_RECEIVE_ABORTED:
+	case QUIC_STREAM_EVENT_IDEAL_SEND_BUFFER_SIZE:
+	case QUIC_STREAM_EVENT_SEND_SHUTDOWN_COMPLETE:
+
+		break;
+
+	default:
+
+		UE_LOG(LogQuicMessagingTransport, Warning,
+			TEXT("[QuicClient] Received unhandled stream event: %d"), Event->Type);
+
+		break;
+
+	}
+
+	return QUIC_STATUS_SUCCESS;
+}
+
+
+void FQuicClient::OnConnectionConnected()
+{
+	ConnectionState = EQuicClientState::Connected;
+
+	EndpointConnectedDelegate.ExecuteIfBound(Endpoint);
+
+	UE_LOG(LogQuicMessagingTransport, Verbose,
+		TEXT("[QuicClient] Handshake completed. Server: %s / LocalNodeId: %s / ThreadId: %d."),
+		*Endpoint.ToString(), *LocalNodeId.ToString(), Thread->GetThreadID());
+}
+
+
+void FQuicClient::OnConnectionShutdownByTransport(QuicUtils::HRESULT Status)
+{
+	ConnectionState = EQuicClientState::Disconnected;
+	bTransportClosed = true;
+
+	EndpointError = QuicUtils::ConvertQuicStatus(Status);
+
+	UE_LOG(LogQuicMessagingTransport, Verbose,
+		TEXT("[QuicClient] Connection end (by transport). Server: %s."),
+		*Endpoint.ToString());
+
+	if (EndpointError != EQuicEndpointError::Normal)
+	{
+		UE_LOG(LogQuicMessagingTransport, Warning,
+			TEXT("[QuicClient] Shutdown reason: %s"),
+			*QuicUtils::GetEndpointErrorString(EndpointError));
+	}
+}
+
+
+void FQuicClient::OnConnectionShutdownByPeer(QUIC_UINT62 ErrorCode)
+{
+	ConnectionState = EQuicClientState::Disconnected;
+	bTransportClosed = true;
+
+	UE_LOG(LogQuicMessagingTransport, Verbose,
+		TEXT("[QuicClient] Connection end (by peer). Server: %s / Error code: %d."),
+		*Endpoint.ToString(), ErrorCode);
+}
+
+
+void FQuicClient::OnConnectionShutdownComplete(HQUIC Connection, BOOLEAN AppCloseInProgress)
+{
+	if (!AppCloseInProgress)
+	{
+		UE_LOG(LogQuicMessagingTransport, Verbose,
+			TEXT("[QuicClient] Connection shutdown complete. Server: %s."),
+			*Endpoint.ToString());
+
+		if (MsQuic && ConnectionState != EQuicClientState::Stopped)
+		{
+			if (ConnectionState != EQuicClientState::Closed)
+			{
+				MsQuic->ConnectionClose(Connection);
+				ConnectionState = EQuicClientState::Closed;
+			}
+
+			EndpointDisconnectedDelegate.ExecuteIfBound(Endpoint, EndpointError);
+		}
+	}
+}
+
+
+void FQuicClient::OnConnectionResumptionTicketReceived()
+{
+	UE_LOG(LogQuicMessagingTransport, Verbose,
+		TEXT("[QuicClient] Resumption ticket received."));
+}
+
+
+void FQuicClient::OnConnectionPeerStreamStarted(HQUIC Stream)
+{
+	const uint64 StreamId = GetStreamId(Stream);
+
+	{
+		const FQuicPayloadPtr InboundData = MakeShared<TArray<uint8>, ESPMode::ThreadSafe>();
+
+		FScopeLock InboundLock(&InboundBuffersCS);
+
+		const QUIC_ADDR LocalAddress = GetAddressFromEndpoint(Endpoint);
+		const FInboundBuffer InboundBuffer(LocalAddress, FIPv4Endpoint::Any, InboundData);
+		const FAddrStream AddressStreamKey(Endpoint, StreamId);
+
+		InboundBuffers.Add(AddressStreamKey, InboundBuffer);
+	}
+
+	UE_LOG(LogQuicMessagingTransport, VeryVerbose,
+		TEXT("[QuicClient] Peer stream started. Server: %s / StreamId: %d."),
+		*Endpoint.ToString(), StreamId);
+}
+
+
+void FQuicClient::OnStreamShutdownComplete(HQUIC Stream) const
+{
+	if (MsQuic)
+	{
+		MsQuic->StreamClose(Stream);
+	}
+}
+
+
+void FQuicClient::SendToStream(const FOutboundMessage& OutboundMessage)
+{
+	HQUIC Stream = nullptr;
+	QUIC_STATUS Status;
+
+	FIPv4Endpoint* PeerAddress = new FIPv4Endpoint(Endpoint);
+
+	FStreamCallbackInfo* StreamCallbackInfo = new FStreamCallbackInfo();
+	StreamCallbackInfo->Endpoint = this;
+	StreamCallbackInfo->PeerAddress = PeerAddress;
+
+	/**
+	 * Create/allocate a new unidirectional stream
+	 */
+	if (QUIC_FAILED(Status = MsQuic->StreamOpen(ClientConnection,
+		QUIC_STREAM_OPEN_FLAG_UNIDIRECTIONAL, ClientStreamCallback, StreamCallbackInfo, &Stream)))
+	{
+		UE_LOG(LogQuicMessagingTransport, Error,
+			TEXT("[QuicClient] Could not open stream: %s."), *QuicUtils::ConvertResult(Status));
+		return;
+	}
+
+	/**
+	 * Start the stream
+	 */
+	if (QUIC_FAILED(Status = MsQuic->StreamStart(Stream, QUIC_STREAM_START_FLAG_NONE)))
+	{
+		UE_LOG(LogQuicMessagingTransport, Error,
+			TEXT("[QuicClient] Could not start stream: %s."), *QuicUtils::ConvertResult(Status));
+
+		MsQuic->StreamClose(Stream);
+		return;
+	}
+
+	const uint64 StreamId = GetStreamId(Stream);
+
+	/**
+	 * Send buffers over the stream
+	 */
+	const FQuicBufferRef MetaBuffer = MakeShared<QUIC_BUFFER, ESPMode::ThreadSafe>();
+	const FQuicBufferRef MessageBuffer = MakeShared<QUIC_BUFFER, ESPMode::ThreadSafe>();
+
+	MetaBuffer->Buffer = OutboundMessage.SerializedHeader->GetData();
+	MetaBuffer->Length = OutboundMessage.SerializedHeader->Num();
+
+	RegisterStreamMessage(StreamId, *PeerAddress,
+		!OutboundMessage.bHasPayload, OutboundMessage, MetaBuffer, MessageBuffer);
+	
+	if (QUIC_FAILED(Status = MsQuic->StreamSend(Stream, &*MetaBuffer, 1,
+		OutboundMessage.bHasPayload ? QUIC_SEND_FLAG_NONE : QUIC_SEND_FLAG_FIN, StreamCallbackInfo)))
+	{
+		UE_LOG(LogQuicMessagingTransport, Error,
+			TEXT("[QuicClient] Could not send meta on stream: %s."), *QuicUtils::ConvertResult(Status));
+		
+		MsQuic->StreamShutdown(Stream, QUIC_STREAM_SHUTDOWN_FLAG_ABORT, 0);
+		return;
+	}
+
+	// This message is header-only
+	if (!OutboundMessage.bHasPayload)
+	{
+		return;
+	}
+
+	MessageBuffer->Buffer = OutboundMessage.SerializedPayload->GetData();
+	MessageBuffer->Length = OutboundMessage.SerializedPayload->Num();
+
+	if (QUIC_FAILED(Status = MsQuic->StreamSend(Stream,
+		&*MessageBuffer, 1, QUIC_SEND_FLAG_FIN, StreamCallbackInfo)))
+	{
+		UE_LOG(LogQuicMessagingTransport, Error,
+			TEXT("[QuicClient] Could not send message on stream: %s."), *QuicUtils::ConvertResult(Status));
+		
+		MsQuic->StreamShutdown(Stream, QUIC_STREAM_SHUTDOWN_FLAG_ABORT, 0);
+	}
+}
+
+
+bool FQuicClient::ReceiveFromStream(const FInboundQuicBuffer& QuicBuffer)
+{
+	FScopeLock InboundLock(&InboundBuffersCS);
+	
+	const FAddrStream AddressStreamKey(QuicBuffer.PeerAddress, QuicBuffer.StreamId);
+
+	FInboundBuffer* Buffer = InboundBuffers.Find(AddressStreamKey);
+
+	if (!Buffer)
+	{
+		UE_LOG(LogQuicMessagingTransport, Error,
+			TEXT("[QuicClient] Invalid inbound buffer in ReceiveFromStream %d."),
+			QuicBuffer.StreamId);
+
+		return false;
+	}
+
+	const FQuicPayloadPtr& InboundData = Buffer->InboundData;
+
+	for (const TPair<uint8*, uint32> InboundBuffer : QuicBuffer.Data)
+	{
+		Buffer->InboundData->Append(InboundBuffer.Key, InboundBuffer.Value);
+	}
+
+	// Retrieve message header
+	if (!Buffer->bHeaderDeserialized
+		&& InboundData->Num() >= QUIC_MESSAGE_HEADER_SIZE)
+	{
+		const FQuicPayloadPtr HeaderData = MakeShared<TArray<uint8>, ESPMode::ThreadSafe>();
+		HeaderData->Append(InboundData->GetData(), QUIC_MESSAGE_HEADER_SIZE);
+		const FMessageHeader MessageHeader = DeserializeHeaderDelegate.Execute(HeaderData);
+
+		Buffer->MessageHeader = MessageHeader;
+		Buffer->bHeaderDeserialized = true;
+
+		InboundData->RemoveAt(0, QUIC_MESSAGE_HEADER_SIZE);
+
+		const uint32& SerializedMessageSize = MessageHeader.SerializedMessageSize;
+		InboundData->Reserve(SerializedMessageSize);
+	}
+
+	InboundLock.Unlock();
+
+	if (MsQuic)
+	{
+		MsQuic->StreamReceiveComplete(QuicBuffer.Stream, QuicBuffer.BufferLength);
+	}
+
+	return true;
+}
+
+
+void FQuicClient::ReachedEndOfStream(const FInboundQuicBuffer& QuicBuffer)
+{
+	FInboundBuffer InboundBuffer;
+
+	{
+		const FAddrStream AddressStreamKey(QuicBuffer.PeerAddress, QuicBuffer.StreamId);
+
+		FScopeLock InboundLock(&InboundBuffersCS);
+
+		if (!InboundBuffers.RemoveAndCopyValue(AddressStreamKey, InboundBuffer))
+		{
+			UE_LOG(LogQuicMessagingTransport, Error,
+				TEXT("[QuicClient] Invalid inbound buffer in ReachedEndOfStream."));
+
+			return;
+		}
+	}
+
+	const FQuicPayloadPtr& InboundData = InboundBuffer.InboundData;
+	const FMessageHeader& MessageHeader = InboundBuffer.MessageHeader;
+
+	// Transmitted message size smaller than expected
+	if (InboundData->Num() < static_cast<int32>(MessageHeader.SerializedMessageSize))
+	{
+		UE_LOG(LogQuicMessagingTransport, Error,
+			TEXT("[QuicClient] Transmitted message size smaller than expected %d / %d / %d."),
+			QuicBuffer.StreamId, InboundData->Num(), MessageHeader.SerializedMessageSize);
+
+		return;
+	}
+
+	FInboundMessage InboundMessage(InboundData, MessageHeader,
+		Endpoint, Endpoint);
+
+	InboundMessages.Enqueue(InboundMessage);
+}
+
diff --git a/Engine/Plugins/Experimental/QuicMessaging/Source/QuicMessagingTransport/Private/QuicClient.h b/Engine/Plugins/Experimental/QuicMessaging/Source/QuicMessagingTransport/Private/QuicClient.h
new file mode 100644
--- /dev/null
+++ b/Engine/Plugins/Experimental/QuicMessaging/Source/QuicMessagingTransport/Private/QuicClient.h
@@ -0,0 +1,223 @@
+// Copyright Epic Games, Inc. All Rights Reserved.
+
+#pragma once
+
+#include "QuicEndpoint.h"
+#include "QuicQueues.h"
+#include "QuicMessages.h"
+
+#include "QuicUtils.h"
+
+
+/**
+ * Implements the QuicClient.
+ */
+class FQuicClient
+	: public FQuicEndpoint
+	, public FQuicQueues
+{
+
+public:
+
+	FQuicClient() = delete;
+
+	/**
+	 * Creates and initializes a new quic client.
+	 */
+	FQuicClient(TSharedRef<FQuicClientConfig> InClientConfig,
+		HQUIC InQuicConfig, const QUIC_API_TABLE* InMsQuic, HQUIC InRegistration,
+		const TCHAR* InThreadName)
+		: FQuicEndpoint(InClientConfig, InQuicConfig,
+			InMsQuic, InRegistration, InThreadName)
+		, ClientConnection(nullptr)
+		, ClientConfig(InClientConfig)
+	{ }
+
+protected:
+
+	/**
+	 * Starts the configured QUIC client.
+	 */
+	virtual void StartEndpoint() override;
+
+public:
+
+	/**
+	 * Stops the QUIC client.
+	 */
+	virtual void StopEndpoint() override;
+
+	/**
+	 * Indicates if the client is connected and ready to send/receive messages.
+	 */
+	virtual bool IsEndpointReady() override;
+
+public:
+
+	/**
+	 * Periodically called to update the endpoint (consume messages).
+	 */
+	virtual void UpdateEndpoint() override;
+
+protected:
+
+	/**
+	 * Dequeue inbound messages until MaxConsumeTime is reached.
+	 * Inbound messages are passed to QuicEndpointManager for validation.
+	 */
+	virtual void ConsumeInboundMessages() override;
+
+	/**
+	 * Dequeue inbound buffers until MaxConsumeTime is reached.
+	 * Inbound buffers via QUIC are copied into our buffers until the stream end is indicated.
+	 */
+	virtual void ProcessInboundBuffers() override;
+
+	/**
+	 * Dequeue outbound messages until MaxConsumeTime is reached.
+	 * Outbound messages will then be sent via QUIC.
+	 */
+	virtual void ConsumeOutboundMessages() override;
+
+protected:
+
+	/**
+	 * Collects, converts and passes QUIC connection stats to the EndpointManager.
+	 */
+	virtual void CollectStatistics() override;
+
+public:
+
+	/**
+	 * Enqueues an outbound message on the client.
+	 *
+	 * @param OutboundMessage The outbound message
+	 */
+	virtual void EnqueueOutboundMessage(const FOutboundMessage& OutboundMessage) override;
+
+	/**
+	 * Sends a hello message to the remote endpoint.
+	 *
+	 * @param RemoteEndpoint The remote endpoint
+	 */
+	virtual void SendHello(FIPv4Endpoint RemoteEndpoint) override;
+
+	/**
+	 * Sends a bye message to the remote endpoint.
+	 */
+	virtual void SendBye() override;
+
+protected:
+
+	/**
+	 * Callback for client connections.
+	 *
+	 * @param Connection The QUIC connection handle
+	 * @param Context The app context (FQuicServer)
+	 * @param Event The connection event
+	 */
+	static QUIC_STATUS ClientConnectionCallback(HQUIC Connection,
+		void* Context, QUIC_CONNECTION_EVENT* Event);
+
+	/**
+	 * Callback for client streams.
+	 *
+	 * @param Stream The QUIC stream handle
+	 * @param Context The app context (FQuicServer)
+	 * @param Event The stream event
+	 */
+	static QUIC_STATUS ClientStreamCallback(HQUIC Stream,
+		void* Context, QUIC_STREAM_EVENT* Event);
+
+protected:
+
+	/**
+	 * Handles a connection after the handshake was completed.
+	 */
+	void OnConnectionConnected();
+
+	/**
+	 * Handles a connection after it was killed due to idle timeout.
+	 *
+	 * @param Status The shutdown event status
+	 */
+	void OnConnectionShutdownByTransport(QuicUtils::HRESULT Status);
+
+	/**
+	 * Handles a connection after it was killed by the peer.
+	 *
+	 * @param ErrorCode The QUIC error code
+	 */
+	void OnConnectionShutdownByPeer(QUIC_UINT62 ErrorCode);
+
+	/**
+	 * Handles a connection after its shutdown was completed, ready for cleanup.
+	 *
+	 * @param Connection The connection handle
+	 * @param AppCloseInProgress Boolean indicating whether app close is already in progress
+	 */
+	void OnConnectionShutdownComplete(HQUIC Connection, BOOLEAN AppCloseInProgress);
+
+	/**
+	 * Handles a received connection resumption ticket.
+	 *
+	 * @note Does not do anything at this time
+	 */
+	void OnConnectionResumptionTicketReceived();
+
+	/**
+	 * Handles when a peer started a stream.
+	 *
+	 * @param Stream The stream handle
+	 */
+	void OnConnectionPeerStreamStarted(HQUIC Stream);
+
+	/**
+	 * Handles a stream when both directions have been shut down, safe for cleanup.
+	 *
+	 * @param Stream The stream handle
+	 */
+	void OnStreamShutdownComplete(HQUIC Stream) const;
+	
+protected:
+
+	/**
+	 * Sends the outbound buffers to the remote endpoint.
+	 *
+	 * @param OutboundMessage The outbound message
+	 */
+	void SendToStream(const FOutboundMessage& OutboundMessage);
+
+protected:
+
+	/**
+	 * Copy data from a stream buffer.
+	 *
+	 * @param QuicBuffer The inbound QUIC buffer
+	 *
+	 * @return Whether the operation succeeded
+	 */
+	bool ReceiveFromStream(const FInboundQuicBuffer& QuicBuffer);
+
+	/**
+	 * Prepares inbound data for processing.
+	 *
+	 * @param QuicBuffer The inbound QUIC buffer
+	 */
+	void ReachedEndOfStream(const FInboundQuicBuffer& QuicBuffer);
+
+protected:
+
+	/** Holds the QUIC API client connection handle. */
+	HQUIC ClientConnection;
+
+	/** Current client connection state. */
+	EQuicClientState ConnectionState = EQuicClientState::Connecting;
+
+	/** Indicates whether the transport is closed. */
+	bool bTransportClosed = false;
+
+	/** Holds the client config. */
+	TSharedRef<FQuicClientConfig> ClientConfig;
+
+};
diff --git a/Engine/Plugins/Experimental/QuicMessaging/Source/QuicMessagingTransport/Private/QuicEndpoint.cpp b/Engine/Plugins/Experimental/QuicMessaging/Source/QuicMessagingTransport/Private/QuicEndpoint.cpp
new file mode 100644
--- /dev/null
+++ b/Engine/Plugins/Experimental/QuicMessaging/Source/QuicMessagingTransport/Private/QuicEndpoint.cpp
@@ -0,0 +1,244 @@
+// Copyright Epic Games, Inc. All Rights Reserved.
+
+#include "QuicEndpoint.h"
+
+
+void FQuicEndpoint::Start()
+{
+	Thread = FRunnableThread::Create(this, *ThreadName, 128 * 4096, TPri_AboveNormal, FPlatformAffinity::GetPoolThreadMask());
+
+	StartEndpoint();
+}
+
+
+uint32 FQuicEndpoint::Run()
+{
+	while (!bStopping)
+	{
+		UpdateEndpoint();
+	}
+
+	return 0;
+}
+
+
+void FQuicEndpoint::Stop()
+{
+	bStopping = true;
+}
+
+
+void FQuicEndpoint::UpdateStatistics()
+{
+	if (!bUpdatingStatistics)
+	{
+		bUpdatingStatistics = true;
+
+		const double CurrentTime = FPlatformTime::Seconds();
+
+		if (CurrentTime < (LastStatsPoll + StatsPollInterval.GetTotalSeconds()))
+		{
+			bUpdatingStatistics = false;
+			return;
+		}
+
+		CollectStatistics();
+
+		LastStatsPoll = FPlatformTime::Seconds();
+		bUpdatingStatistics = false;
+	}
+}
+
+
+uint64 FQuicEndpoint::GetStreamId(HQUIC Stream) const
+{
+	if (MsQuic == nullptr)
+	{
+		return 0;
+	}
+
+	uint64_t StreamId;
+	uint32_t StreamIdSize = sizeof(StreamId);
+
+	QUIC_STATUS Status;
+
+	if (QUIC_FAILED(Status = MsQuic->GetParam(Stream,
+		QUIC_PARAM_STREAM_ID, &StreamIdSize, &StreamId)))
+	{
+		return 0;
+	}
+
+	return StreamId;
+}
+
+
+/*
+ * TODO(vri): Fully implement and fix statistics
+ * Tracked in [UCS-5152] Finalize QuicMessaging plugin
+ */
+QUIC_STATUS FQuicEndpoint::GetConnectionStats(
+	HQUIC Connection, QUIC_STATISTICS& ConnectionStats)
+{
+	uint32_t StatsSize = sizeof(ConnectionStats);
+
+	QUIC_STATUS Status = QUIC_STATUS_SUCCESS;
+
+	/*Status = MsQuic->GetParam(Connection,
+		QUIC_PARAM_CONN_STATISTICS, &StatsSize, &ConnectionStats);*/
+
+	return Status;
+}
+
+
+/*
+ * TODO(vri): Fully implement and fix statistics
+ * Tracked in [UCS-5152] Finalize QuicMessaging plugin
+ */
+FMessageTransportStatistics FQuicEndpoint::ConvertStatistics(
+	FIPv4Endpoint StatsEndpoint, QUIC_STATISTICS& QuicStats)
+{
+	FMessageTransportStatistics TransportStats;
+
+	TransportStats.TotalBytesSent = QuicStats.Send.TotalBytes;
+	TransportStats.PacketsSent = QuicStats.Send.TotalPackets;
+	TransportStats.PacketsLost = QuicStats.Send.SuspectedLostPackets - QuicStats.Send.SpuriousLostPackets;
+	TransportStats.PacketsAcked = QuicStats.Send.TotalPackets;
+	TransportStats.PacketsReceived = QuicStats.Recv.TotalPackets;
+	TransportStats.PacketsInFlight = 0;
+	TransportStats.WindowSize = 0; // todo
+	TransportStats.AverageRTT = { 0 }; // todo
+	TransportStats.IPv4AsString = StatsEndpoint.ToString();
+
+	return TransportStats;
+}
+
+
+void FQuicEndpoint::OnStreamSendComplete(const HQUIC Stream, void* Context)
+{
+	const uint64 StreamId = GetStreamId(Stream);
+	const FIPv4Endpoint* PeerAddress = static_cast<FIPv4Endpoint*>(Context);
+	const FAddrStream AddressStreamKey(*PeerAddress, StreamId);
+
+	FScopeLock StreamMessagesLock(&StreamMessagesCS);
+
+	FMessageMeta* MessageMeta = StreamMessages.Find(AddressStreamKey);
+
+	if (!MessageMeta)
+	{
+		return;
+	}
+
+	if (!MessageMeta->bIsMetaOnly && !MessageMeta->bMetaSent)
+	{
+		MessageMeta->bMetaSent = true;
+		return;
+	}
+
+	StreamMessages.Remove(AddressStreamKey);
+}
+
+
+void FQuicEndpoint::OnStreamReceive(const HQUIC Stream,
+	const FIPv4Endpoint PeerAddress, const QUIC_STREAM_EVENT* Event)
+{
+	const bool bEndOfStream = (Event->RECEIVE.Flags & QUIC_RECEIVE_FLAG_FIN) != 0;
+
+	const uint64 StreamId = GetStreamId(Stream);
+
+	FInboundQuicBuffer InboundBuffer(
+		StreamId, Stream, bEndOfStream, PeerAddress, Event->RECEIVE.TotalBufferLength);
+
+	for (uint32 i = 0; i < Event->RECEIVE.BufferCount; ++i)
+	{
+		const uint32 BufferLength = Event->RECEIVE.Buffers[i].Length;
+
+		if (BufferLength < 1)
+		{
+			UE_LOG(LogTemp, Warning,
+				TEXT("[QuicEndpoint] BufferLength less than 1 on streamId %d."), StreamId);
+
+			continue;
+		}
+
+		InboundBuffer.Data.Add(Event->RECEIVE.Buffers[i].Buffer, BufferLength);
+	}
+
+	InboundQuicBuffers.Enqueue(InboundBuffer);
+}
+
+
+void FQuicEndpoint::OnStreamPeerSendAborted(const HQUIC Stream) const
+{
+	if (MsQuic)
+	{
+		MsQuic->StreamShutdown(Stream, QUIC_STREAM_SHUTDOWN_FLAG_ABORT, 0);
+	}
+}
+
+
+void FQuicEndpoint::RegisterStreamMessage(
+	const uint64 StreamId, const FIPv4Endpoint PeerAddress,
+	const bool bIsMetaOnly, const FOutboundMessage& OutboundMessage,
+	const FQuicBufferRef& MetaBuffer, const FQuicBufferRef& MessageBuffer)
+{
+	const FMessageMeta MessageMeta(
+		bIsMetaOnly, OutboundMessage, MetaBuffer, MessageBuffer);
+
+	const FAddrStream AddressStreamKey(PeerAddress, StreamId);
+
+	FScopeLock StreamMessagesLock(&StreamMessagesCS);
+
+	StreamMessages.Add(AddressStreamKey, MessageMeta);
+}
+
+
+QUIC_ADDR FQuicEndpoint::GetAddressFromEndpoint(FIPv4Endpoint RemoteEndpoint) const
+{
+	QUIC_ADDR RemoteAddress = {{0}};
+
+#ifdef _WIN32
+
+	inet_pton(AF_INET, TCHAR_TO_ANSI(*RemoteEndpoint.Address.ToString()),
+		&RemoteAddress.Ipv4.sin_addr.S_un.S_addr);
+
+#elif defined(__linux__) || defined(__APPLE__)
+
+	inet_pton(AF_INET, TCHAR_TO_ANSI(*RemoteEndpoint.Address.ToString()),
+		&RemoteAddress.Ipv4.sin_addr.s_addr);
+
+#endif
+
+	QuicAddrSetFamily(&RemoteAddress, QUIC_ADDRESS_FAMILY_INET);
+	QuicAddrSetPort(&RemoteAddress, RemoteEndpoint.Port);
+
+	return RemoteAddress;
+}
+
+
+FIPv4Endpoint FQuicEndpoint::GetEndpointFromAddress(QUIC_ADDR RemoteAddress) const
+{
+	FIPv4Endpoint RemoteEndpoint;
+
+#ifdef _WIN32
+
+	char Ipv4Buf[INET_ADDRSTRLEN];
+
+	inet_ntop(AF_INET, &RemoteAddress.Ipv4.sin_addr.S_un.S_addr, Ipv4Buf, sizeof(Ipv4Buf));
+
+	RemoteEndpoint.Address.Parse(ANSI_TO_TCHAR(Ipv4Buf), RemoteEndpoint.Address);
+
+#elif defined(__linux__) || defined(__APPLE__)
+
+	char Ipv4Buf[INET_ADDRSTRLEN];
+
+	inet_ntop(AF_INET, &RemoteAddress.Ipv4.sin_addr.s_addr, Ipv4Buf, sizeof(Ipv4Buf));
+
+	RemoteEndpoint.Address.Parse(ANSI_TO_TCHAR(Ipv4Buf), RemoteEndpoint.Address);
+
+#endif
+
+	RemoteEndpoint.Port = QuicAddrGetPort(&RemoteAddress);
+
+	return RemoteEndpoint;
+}
+
diff --git a/Engine/Plugins/Experimental/QuicMessaging/Source/QuicMessagingTransport/Private/QuicEndpoint.h b/Engine/Plugins/Experimental/QuicMessaging/Source/QuicMessagingTransport/Private/QuicEndpoint.h
new file mode 100644
--- /dev/null
+++ b/Engine/Plugins/Experimental/QuicMessaging/Source/QuicMessagingTransport/Private/QuicEndpoint.h
@@ -0,0 +1,502 @@
+// Copyright Epic Games, Inc. All Rights Reserved.
+
+#pragma once
+
+#include "QuicMessagingTransportPrivate.h"
+
+#include "QuicIncludes.h"
+
+#include "CoreMinimal.h"
+#include "Containers/Queue.h"
+#include "HAL/PlatformTime.h"
+#include "HAL/Runnable.h"
+#include "HAL/RunnableThread.h"
+#include "Misc/SingleThreadRunnable.h"
+#include "Interfaces/IPv4/IPv4Endpoint.h"
+#include "INetworkMessagingExtension.h"
+
+#include "QuicFlags.h"
+#include "QuicMessages.h"
+#include "QuicBuffers.h"
+#include "QuicEndpointConfig.h"
+#include "QuicUtils.h"
+
+
+using FQuicBufferRef = TSharedRef<QUIC_BUFFER, ESPMode::ThreadSafe>;
+using FAddrStream = TPair<FIPv4Endpoint, uint64>;
+
+
+DECLARE_DELEGATE_OneParam(FOnMessageInbound, const FInboundMessage /*InboundMessage*/)
+
+DECLARE_DELEGATE_RetVal_OneParam(FQuicPayloadPtr, FOnSerializeHeader, FMessageHeader& /*MessageHeader*/)
+
+DECLARE_DELEGATE_RetVal_OneParam(FMessageHeader, FOnDeserializeHeader, FQuicPayloadPtr /*HeaderData*/)
+
+DECLARE_DELEGATE_TwoParams(FOnEndpointDisconnected, const FIPv4Endpoint /*DisconnectedEndpoint*/, EQuicEndpointError /*Reason*/)
+
+DECLARE_DELEGATE_OneParam(FOnEndpointConnected, const FIPv4Endpoint& /*ConnectedEndpoint*/)
+
+DECLARE_DELEGATE_OneParam(FOnStatisticsUpdated, const FMessageTransportStatistics /*TransportStats*/)
+
+
+/**
+ * The base class for FQuicServers and FQuicClients.
+ */
+class FQuicEndpoint
+	: public FRunnable
+	, private FSingleThreadRunnable
+{
+
+protected:
+
+	FQuicEndpoint() = delete;
+
+	/**
+	 * Creates and initializes a quic endpoint.
+	 * 
+	 * @param InEndpointConfig The endpoint config
+	 * @param InQuicConfig The QUIC API endpoint config
+	 * @param InMsQuic The QUIC API table handle
+	 * @param InRegistration The QUIC API registration handle
+	 * @param InThreadName The thread name
+	 */
+	FQuicEndpoint(TSharedRef<FQuicEndpointConfig> InEndpointConfig,
+		HQUIC InQuicConfig, const QUIC_API_TABLE* InMsQuic,
+		HQUIC InRegistration, const TCHAR* InThreadName)
+		: Endpoint(InEndpointConfig->Endpoint)
+		, LocalNodeId(InEndpointConfig->LocalNodeId)
+		, QuicConfig(InQuicConfig)
+		, MsQuic(InMsQuic)
+		, Registration(InRegistration)
+		, bStopping(false)
+		, Thread(nullptr)
+		, ThreadName(InThreadName)
+		, InboundBuffers(TMap<FAddrStream, FInboundBuffer>())
+	{ }
+
+	/**
+	 * Destruct QUIC endpoint.
+	 */
+	virtual ~FQuicEndpoint()
+	{
+		if (Thread != nullptr)
+		{
+			Thread->Kill(true);
+			delete Thread;
+		}
+	}
+
+public:
+
+	/**
+	 * Start the thread and QUIC endpoint.
+	 */
+	void Start();
+
+public:
+
+	/**
+	 * Continuously updates the endpoint.
+	 */
+	uint32 Run() override;
+
+	void Stop() override;
+
+protected:
+
+	/**
+	 * Starts the configured QUIC endpoint.
+	 */
+	virtual void StartEndpoint() = 0;
+
+public:
+
+	/**
+	 * Stops the QUIC endpoint.
+	 */
+	virtual void StopEndpoint() = 0;
+
+	/**
+	 * Indicates if the endpoint is ready to send and receive messages.
+	 */
+	virtual bool IsEndpointReady() = 0;
+
+	/**
+	 * Indicates whether the endpoint is stopping.
+	 */
+	bool IsStopping() const
+	{
+		return bStopping;
+	}
+
+	/**
+	 * Checks if a QuicServer has a handler with this endpoint.
+	 *
+	 * @param HandlerEndpoint The handler endpoint
+	 * @note This is only valid for QuicServers
+	 */
+	virtual bool HasHandler(const FIPv4Endpoint HandlerEndpoint) { return false; }
+
+public:
+
+	/**
+	 * Periodically called to update the endpoint (consume messages and ping).
+	 */
+	virtual void UpdateEndpoint() = 0;
+
+protected:
+
+	/**
+	 * Periodically called within UpdateEndpoint to updated QUIC connection stats.
+	 */
+	void UpdateStatistics();
+
+	/**
+	 * Collects, converts and passes QUIC connection stats to the EndpointManager.
+	 */
+	virtual void CollectStatistics() = 0;
+
+public:
+
+	/**
+	 * Marks a handler of a QuicServer as authenticated.
+	 * 
+	 * @param Handler The handler endpoint
+	 * @note This is only valid for QuicServers
+	 */
+	virtual void SetHandlerAuthenticated(const FIPv4Endpoint& Handler) { }
+
+	/**
+	 * Disconnects a handler.
+	 *
+	 * @param HandlerEndpoint The handler endpoint
+	 * @note This is only valid for QuicServer
+	 */
+	virtual void DisconnectHandler(const FIPv4Endpoint& HandlerEndpoint) { }
+
+public:
+
+	/**
+	 * Enqueues an outbound message on the endpoint.
+	 *
+	 * @param OutboundMessage The outbound message
+	 */
+	virtual void EnqueueOutboundMessage(const FOutboundMessage& OutboundMessage) = 0;
+
+	/**
+	 * Sends a hello message to the remote endpoint.
+	 *
+	 * @param RemoteEndpoint The remote endpoint
+	 */
+	virtual void SendHello(FIPv4Endpoint RemoteEndpoint) = 0;
+
+	/**
+	 * Sends a bye message to the remote endpoint(s).
+	 */
+	virtual void SendBye() = 0;
+
+protected:
+
+	/**
+	 * Helper to get the streamId from a QUIC stream.
+	 *
+	 * @param Stream The QUIC stream handle
+	 */
+	uint64 GetStreamId(HQUIC Stream) const;
+
+	/**
+	 * Helper to get the connection level stats.
+	 *
+	 * @param Connection The connection handle
+	 * @param ConnectionStats The QUIC statistics container
+	 */
+	QUIC_STATUS GetConnectionStats(
+		HQUIC Connection, QUIC_STATISTICS& ConnectionStats);
+
+	/**
+	 * Helper to convert QUIC endpoint statistics to MessageTransportStatistics.
+	 *
+	 * @param StatsEndpoint The endpoint which the statistics originate from
+	 * @param QuicStats The QUIC endpoint statistics
+	 */
+	FMessageTransportStatistics ConvertStatistics(
+		FIPv4Endpoint StatsEndpoint, QUIC_STATISTICS& QuicStats);
+
+protected:
+
+	/**
+	 * Struct for QuicCallbacks to pass information to event handlers.
+	 */
+	struct FStreamCallbackInfo
+	{
+		/** Holds a pointer to the peer address. */
+		FIPv4Endpoint* PeerAddress;
+
+		/** Holds a pointer to either the FQuicServer or FQuicClient. */
+		FQuicEndpoint* Endpoint;
+	};
+
+protected:
+
+	/**
+	 * Handles a stream when all data has been sent and the buffers can be released.
+	 *
+	 * @param Stream The stream handle
+	 * @param Context The app context (QUIC_ADDR)
+	 */
+	void OnStreamSendComplete(const HQUIC Stream, void* Context);
+
+	/**
+	 * Handles received data from the peer.
+	 *
+	 * @param Stream The stream handle
+	 * @param PeerAddress The peer address
+	 * @param Event The stream event
+	 */
+	void OnStreamReceive(const HQUIC Stream,
+		const FIPv4Endpoint PeerAddress, const QUIC_STREAM_EVENT* Event);
+
+	/**
+	 * Handles a stream when the peer has aborted sending.
+	 *
+	 * @param Stream The stream handle
+	 */
+	void OnStreamPeerSendAborted(const HQUIC Stream) const;
+
+protected:
+
+	/**
+	 * Register a stream with a MessageId and buffers for cleanup reference.
+	 *
+	 * @param StreamId The stream ID
+	 * @param PeerAddress The peer address
+	 * @param bIsMetaOnly Flag indicating whether the message is meta-only
+	 * @param OutboundMessage The outbound message
+	 * @param MetaBuffer SharedRef to the meta buffer
+	 * @param MessageBuffer SharedRef to the message buffer
+	 */
+	void RegisterStreamMessage(
+		const uint64 StreamId, const FIPv4Endpoint PeerAddress,
+		const bool bIsMetaOnly, const FOutboundMessage& OutboundMessage,
+		const FQuicBufferRef& MetaBuffer, const FQuicBufferRef& MessageBuffer);
+
+protected:
+
+	/**
+	 * Convert an endpoint to a QUIC address.
+	 *
+	 * @param RemoteEndpoint The endpoint to convert
+	 */
+	QUIC_ADDR GetAddressFromEndpoint(FIPv4Endpoint RemoteEndpoint) const;
+
+	/**
+	 * Convert a QUIC address to an endpoint.
+	 *
+	 * @param RemoteAddress The address to convert
+	 */
+	FIPv4Endpoint GetEndpointFromAddress(QUIC_ADDR RemoteAddress) const;
+
+public:
+
+	/**
+	 * Delegate for when a message has been received by the endpoint and needs to be checked.
+	 *
+	 * This delegate must be bound before the endpoint thread is started with
+	 * the Start() method. It cannot be unbound while the thread is running.
+	 * 
+	 * @return The delegate
+	 */
+	FOnMessageInbound& OnMessageInbound()
+	{
+		return MessageInboundDelegate;
+	}
+
+public:
+
+	/**
+	 * Delegate for when a header needs to be serialized.
+	 *
+	 * This delegate must be bound before the endpoint thread is started with
+	 * the Start() method. It cannot be unbound while the thread is running.
+	 * 
+	 * @return The delegate
+	 */
+	FOnSerializeHeader& OnSerializeHeader()
+	{
+		return SerializeHeaderDelegate;
+	}
+
+	/**
+	 * Delegate for when a header needs to be deserialized.
+	 *
+	 * This delegate must be bound before the endpoint thread is started with
+	 * the Start() method. It cannot be unbound while the thread is running.
+	 *
+	 * @return The delegate
+	 */
+	FOnDeserializeHeader& OnDeserializeHeader()
+	{
+		return DeserializeHeaderDelegate;
+	}
+
+public:
+
+	/**
+	 * Delegate for when an endpoint has disconnected or the connection failed.
+	 *
+	 * This delegate must be bound before the endpoint thread is started with
+	 * the Start() method. It cannot be unbound while the thread is running.
+	 *
+	 * @return The delegate
+	 */
+	FOnEndpointDisconnected& OnEndpointDisconnected()
+	{
+		return EndpointDisconnectedDelegate;
+	}
+
+public:
+
+	/**
+	 * Delegate for when an endpoint has connected.
+	 *
+	 * This delegate must be bound before the endpoint thread is started with
+	 * the Start() method. It cannot be unbound while the thread is running.
+	 *
+	 * @return The delegate
+	 */
+	FOnEndpointConnected& OnEndpointConnected()
+	{
+		return EndpointConnectedDelegate;
+	}
+
+public:
+
+	/**
+	 * Delegate for when the statistics for an endpoint are updated.
+	 *
+	 * This delegate must be bound before the endpoint thread is started with
+	 * the Start() method. It cannot be unbound while the thread is running.
+	 *
+	 * @return The delegate
+	 */
+	FOnStatisticsUpdated& OnStatisticsUpdated()
+	{
+		return StatisticsUpdatedDelegate;
+	}
+
+protected:
+
+	//~ FSingleThreadRunnable interface
+
+	virtual void Tick() override
+	{
+		return;
+	}
+
+protected:
+
+	/** Holds the IP address (and port) that the socket will be bound to. */
+	FIPv4Endpoint Endpoint;
+
+	/** Holds the local node identifier. */
+	FGuid LocalNodeId;
+
+	/** Holds the QUIC API config. */
+	HQUIC QuicConfig;
+
+	/** Holds the QUIC API handle. */
+	const QUIC_API_TABLE* MsQuic;
+
+	/** Holds the QUIC API registration handle. */
+	HQUIC Registration;
+
+	/** Flag indicating that thread is stopping. */
+	bool bStopping;
+
+	/** The thread object. */
+	FRunnableThread* Thread;
+
+	/** The endpoint thread's name. */
+	FString ThreadName;
+
+	/** Current endpoint error. */
+	EQuicEndpointError EndpointError = EQuicEndpointError::Normal;
+
+protected:
+
+	struct FMessageMeta
+	{
+		/** Holds a flag indicating whether there is only a meta buffer. */
+		bool bIsMetaOnly;
+
+		/** Holds a flag indicating whether the meta buffer has been sent. */
+		bool bMetaSent;
+
+		/** Holds the outbound message. */
+		FOutboundMessage OutboundMessage;
+
+		/** Holds the meta buffer. */
+		FQuicBufferRef MetaBuffer;
+
+		/** Holds the message buffer. */
+		FQuicBufferRef MessageBuffer;
+		
+		/** Creates and initializes a new instance. */
+		FMessageMeta(
+			const bool bInIsMetaOnly, const FOutboundMessage InOutboundMessage,
+			const FQuicBufferRef InMetaBuffer, const FQuicBufferRef InMessageBuffer)
+			: bIsMetaOnly(bInIsMetaOnly)
+			, bMetaSent(false)
+			, OutboundMessage(InOutboundMessage)
+			, MetaBuffer(InMetaBuffer)
+			, MessageBuffer(InMessageBuffer)
+		{ }
+	};
+
+	/** Holds QUIC streamIds and associated messages. */
+	TMap<FAddrStream, FMessageMeta> StreamMessages;
+
+	/** Holds the StreamMessages critical section. */
+	mutable FCriticalSection StreamMessagesCS;
+
+	/** Holds the inbound buffers. */
+	TMap<FAddrStream, FInboundBuffer> InboundBuffers;
+
+	/** Holds the InboundBuffers critical section. */
+	mutable FCriticalSection InboundBuffersCS;
+
+	/** Holds the inbound MsQuic buffers. */
+	TQueue<FInboundQuicBuffer> InboundQuicBuffers;
+
+protected:
+
+	/** Flag indicating if statistics are currently being updated. */
+	bool bUpdatingStatistics = false;
+
+	/** Holds timestamp of when statistics have last been updated. */
+	double LastStatsPoll = FPlatformTime::Seconds();
+
+	/** Holds the interval for statistics updates. */
+	const FTimespan StatsPollInterval = 10000000;
+
+protected:
+
+	/** Holds the message inbound delegate. */
+	FOnMessageInbound MessageInboundDelegate;
+
+	/** Holds the serialize header delegate. */
+	FOnSerializeHeader SerializeHeaderDelegate;
+
+	/** Holds the deserialize header delegate. */
+	FOnDeserializeHeader DeserializeHeaderDelegate;
+
+	/** Holds the endpoint disconnected delegate. */
+	FOnEndpointDisconnected EndpointDisconnectedDelegate;
+
+	/** Holds the endpoint connected delegate. */
+	FOnEndpointConnected EndpointConnectedDelegate;
+
+	/** Holds the statistics updated delegate. */
+	FOnStatisticsUpdated StatisticsUpdatedDelegate;
+
+};
diff --git a/Engine/Plugins/Experimental/QuicMessaging/Source/QuicMessagingTransport/Private/QuicEndpointManager.cpp b/Engine/Plugins/Experimental/QuicMessaging/Source/QuicMessagingTransport/Private/QuicEndpointManager.cpp
new file mode 100644
--- /dev/null
+++ b/Engine/Plugins/Experimental/QuicMessaging/Source/QuicMessagingTransport/Private/QuicEndpointManager.cpp
@@ -0,0 +1,1013 @@
+// Copyright Epic Games, Inc. All Rights Reserved.
+
+#include "QuicEndpointManager.h"
+#include "QuicMessagingTransportPrivate.h"
+
+#include "QuicEndpoint.h"
+#include "QuicClient.h"
+#include "QuicServer.h"
+#include "QuicBuffers.h"
+#include "QuicCertificate.h"
+#include "QuicIncludes.h"
+#include "QuicEndpointConfig.h"
+
+#include "Math/NumericLimits.h"
+
+#include <stdexcept>
+
+
+/** The protocol name used in the Application Layer Protocol Negotation. */
+const QUIC_BUFFER Alpn = { sizeof("unrealquic") - 1, (uint8_t*)"unrealquic" };
+
+
+FQuicEndpointManager::FQuicEndpointManager(
+	TSharedRef<FQuicEndpointConfig> InEndpointConfig)
+    : MsQuic(nullptr)
+	, Registration(nullptr)
+	, CurrentConfiguration(nullptr)
+	, LocalNodeId(InEndpointConfig->LocalNodeId)
+	, ServerEndpoint(FIPv4Endpoint::Any)
+	, EndpointMode(EEndpointMode::Client)
+	, EndpointConfig(InEndpointConfig)
+	, ReconnectAttempts(TMap<FIPv4Endpoint, uint32>())
+	, Endpoints(TMap<FIPv4Endpoint, FQuicEndpointPtr>())
+	, KnownNodes(TMap<FGuid, FQuicNodeInfo>())
+{
+	QUIC_STATUS Status;
+    EndpointMode = EEndpointMode::Client;
+
+	if (QUIC_FAILED(Status = MsQuicOpen2(&MsQuic)))
+	{
+		UE_LOG(LogQuicMessagingTransport, Error,
+			TEXT("[EndpointManager] Could not open MsQuic: %s."),
+			*QuicUtils::ConvertResult(Status));
+
+		return;
+	}
+
+	constexpr QUIC_REGISTRATION_CONFIG RegistrationConfig
+		= { "QuicMessaging", QUIC_EXECUTION_PROFILE_LOW_LATENCY };
+
+	if (QUIC_FAILED(Status = MsQuic->RegistrationOpen(&RegistrationConfig,
+		&Registration)))
+	{
+		UE_LOG(LogQuicMessagingTransport, Error,
+			TEXT("[EndpointManager] Could not open registration: %s."),
+			*QuicUtils::ConvertResult(Status));
+
+		return;
+	}
+
+	// Setup the discovery timeout tick
+	if (EndpointConfig->DiscoveryTimeoutSec > 0)
+	{
+		DiscoveryTimeoutTickHandle = FTSTicker::GetCoreTicker().AddTicker(
+			TEXT("QuicDiscoveryTimeoutTick"), 1, [this](float DeltaSeconds) {
+
+			CheckDiscoveryTimeout();
+			return true;
+
+		});
+	}
+
+	UE_LOG(LogQuicMessagingTransport, Verbose, 
+		TEXT("[EndpointManager] Started endpoint manager."));
+}
+
+
+FQuicEndpointManager::~FQuicEndpointManager()
+{
+	Shutdown();
+
+	if (DiscoveryTimeoutTickHandle.IsValid())
+	{
+		FTSTicker::GetCoreTicker().RemoveTicker(DiscoveryTimeoutTickHandle);
+	}
+	
+	if (MsQuic)
+	{
+		if (CurrentConfiguration)
+		{
+			MsQuic->ConfigurationClose(CurrentConfiguration);
+		}
+
+		if (Registration)
+		{
+			MsQuic->RegistrationShutdown(Registration, QUIC_CONNECTION_SHUTDOWN_FLAG_SILENT, 0);
+			MsQuic->RegistrationClose(Registration);
+		}
+
+		MsQuicClose(MsQuic);
+	}
+
+	UE_LOG(LogQuicMessagingTransport, Verbose, 
+		TEXT("[EndpointManager] Stopped endpoint manager."));
+}
+
+
+void FQuicEndpointManager::Shutdown()
+{
+	FScopeLock EndpointsLock(&EndpointsCS);
+
+	if (Endpoints.IsEmpty())
+	{
+		return;
+	}
+
+	for (auto It = Endpoints.CreateIterator(); It; ++It)
+	{
+		const TSharedPtr<FQuicEndpoint, ESPMode::ThreadSafe>& Endpoint = It->Value;
+
+		if (Endpoint.IsValid())
+		{
+			if (!Endpoint->IsStopping())
+			{
+				Endpoint->SendBye();
+				Endpoint->StopEndpoint();
+
+				CheckLoseNode(It->Key);
+			}
+		}
+
+		if (EndpointMode != EEndpointMode::Server)
+		{
+			It.RemoveCurrent();
+		}
+	}
+
+	UE_LOG(LogQuicMessagingTransport, Verbose, 
+		TEXT("[EndpointManager] Shutdown completed."));
+}
+
+
+HQUIC FQuicEndpointManager::LoadClientConfiguration() const
+{
+	QUIC_SETTINGS QuicSettings = { {0} };
+
+	/** Client idle timeout disabled */
+	QuicSettings.IdleTimeoutMs = 0;
+	QuicSettings.IsSet.IdleTimeoutMs = 1;
+
+	/** Client connection timeout (5 seconds) */
+	QuicSettings.DisconnectTimeoutMs = 5000;
+	QuicSettings.IsSet.DisconnectTimeoutMs = 1;
+
+	/** Allow peer to open max unidirectional streams */
+	QuicSettings.PeerUnidiStreamCount = std::numeric_limits<uint16_t>::max();
+	QuicSettings.IsSet.PeerUnidiStreamCount = 1;
+
+	/** Disable sendbuffering to avoid copying outbound data and delays in sending. */
+	QuicSettings.SendBufferingEnabled = 0;
+	QuicSettings.IsSet.SendBufferingEnabled = 1;
+
+	/**
+	 * Default client config
+	 */
+	QUIC_CREDENTIAL_CONFIG CredConfig;
+	FMemory::Memset(&CredConfig, 0, sizeof(CredConfig));
+
+	CredConfig.Type = QUIC_CREDENTIAL_TYPE_NONE;
+	CredConfig.Flags = QUIC_CREDENTIAL_FLAG_CLIENT;
+
+	const TSharedRef<FQuicClientConfig> ClientEndpointConfig
+		= StaticCastSharedRef<FQuicClientConfig, FQuicEndpointConfig>(EndpointConfig);
+
+	if (ClientEndpointConfig->ClientVerificationMode == EQuicClientVerificationMode::Pass)
+	{
+		CredConfig.Flags |= QUIC_CREDENTIAL_FLAG_NO_CERTIFICATE_VALIDATION;
+	}
+
+	/**
+	 * Initialize client config object
+	 */
+	HQUIC ClientConfiguration;
+
+	QUIC_STATUS Status = QUIC_STATUS_SUCCESS;
+
+	if (QUIC_FAILED(Status = MsQuic->ConfigurationOpen(Registration, &Alpn, 1,
+		&QuicSettings, sizeof(QuicSettings), NULL, &ClientConfiguration)))
+	{
+		UE_LOG(LogQuicMessagingTransport, Error,
+			TEXT("[EndpointManager] Could not open client configuration: %s."),
+			*QuicUtils::ConvertResult(Status));
+
+		return nullptr;
+	}
+
+	/**
+	 * Load TLS credential part of configuration
+	 */
+	if (QUIC_FAILED(Status = MsQuic->ConfigurationLoadCredential(ClientConfiguration, &CredConfig)))
+	{
+		UE_LOG(LogQuicMessagingTransport, Error,
+			TEXT("[EndpointManager] Could not load client config credentials: %s."),
+			*QuicUtils::ConvertResult(Status));
+
+		return nullptr;
+	}
+
+	return ClientConfiguration;
+}
+
+
+HQUIC FQuicEndpointManager::LoadServerConfiguration(
+	FString CertificatePath, FString PrivateKeyPath) const
+{
+	QUIC_SETTINGS QuicSettings = { {0} };
+
+	/** Server idle timeout disabled */
+	QuicSettings.IdleTimeoutMs = 0;
+	QuicSettings.IsSet.IdleTimeoutMs = 1;
+
+	/** Server connection timeout (5 seconds) */
+	QuicSettings.DisconnectTimeoutMs = 5000;
+	QuicSettings.IsSet.DisconnectTimeoutMs = 1;
+
+	/** Server resumption level (allow and 0-RTT) */
+	QuicSettings.ServerResumptionLevel = QUIC_SERVER_RESUME_AND_ZERORTT;
+	QuicSettings.IsSet.ServerResumptionLevel = 1;
+
+	/** Allow peer to open max unidirectional streams */
+	QuicSettings.PeerUnidiStreamCount = std::numeric_limits<uint16_t>::max();
+	QuicSettings.IsSet.PeerUnidiStreamCount = 1;
+
+	/** KeepAlive interval (300 milliseconds) */
+	QuicSettings.KeepAliveIntervalMs = 300;
+	QuicSettings.IsSet.KeepAliveIntervalMs = 1;
+
+	/** Disable sendbuffering to avoid copying outbound data and delays in sending. */
+	QuicSettings.SendBufferingEnabled = 0;
+	QuicSettings.IsSet.SendBufferingEnabled = 1;
+
+	/**
+	 * Credentials
+	 */
+	QuicCertificateUtils::QUIC_CREDENTIAL_CONFIG_HELPER Config;
+	FMemory::Memset(&Config, 0, sizeof(Config));
+
+	Config.CredConfig.Flags = QUIC_CREDENTIAL_FLAG_NONE;
+
+	auto CertFile = StringCast<ANSICHAR>(*CertificatePath);
+	auto KeyFile = StringCast<ANSICHAR>(*PrivateKeyPath);
+
+	Config.CertificateFile.CertificateFile = CertFile.Get();
+	Config.CertificateFile.PrivateKeyFile = KeyFile.Get();
+	Config.CredConfig.Type = QUIC_CREDENTIAL_TYPE_CERTIFICATE_FILE;
+	Config.CredConfig.CertificateFile = &Config.CertificateFile;
+
+	/**
+	 * Initialize server config object
+	 */
+	HQUIC ServerConfiguration;
+
+	QUIC_STATUS Status = QUIC_STATUS_SUCCESS;
+
+	if (QUIC_FAILED(Status = MsQuic->ConfigurationOpen(Registration, &Alpn, 1,
+		&QuicSettings, sizeof(QuicSettings), NULL, &ServerConfiguration)))
+	{
+		UE_LOG(LogQuicMessagingTransport, Error,
+			TEXT("[EndpointManager] Could not open server configuration: %s."),
+			*QuicUtils::ConvertResult(Status));
+
+		return nullptr;
+	}
+
+	if (QUIC_FAILED(Status = MsQuic->ConfigurationLoadCredential(ServerConfiguration, &Config.CredConfig)))
+	{
+		UE_LOG(LogQuicMessagingTransport, Error,
+			TEXT("[EndpointManager] Could not load server config credentials: %s."),
+			*QuicUtils::ConvertResult(Status));
+
+		return nullptr;
+	}
+
+	return ServerConfiguration;
+}
+
+
+void FQuicEndpointManager::InitializeServer()
+{
+	const TSharedRef<FQuicServerConfig> ServerEndpointConfig
+		= StaticCastSharedRef<FQuicServerConfig, FQuicEndpointConfig>(EndpointConfig);
+
+	ServerEndpoint = ServerEndpointConfig->Endpoint;
+
+	if (ServerEndpointConfig->AuthenticationMode == EAuthenticationMode::Disabled)
+	{
+		UE_LOG(LogQuicMessagingTransport, Warning, 
+			TEXT("[EndpointManager] Authentication disabled."));
+	}
+	else
+	{
+		UE_LOG(LogQuicMessagingTransport, Verbose, 
+			TEXT("[EndpointManager] Authentication enabled."));
+	}
+
+	if (ServerEndpointConfig->ConnCooldownMode == EConnectionCooldownMode::Enabled)
+	{
+		UE_LOG(LogQuicMessagingTransport, Verbose,
+			TEXT("[EndpointManager] Connection cooldown active: "
+			"MaxAttempts: %d, TimePeriod: %ds, Cooldown: %ds, MaxCooldown: %ds."),
+			ServerEndpointConfig->ConnCooldownMaxAttempts,
+			ServerEndpointConfig->ConnCooldownPeriodSec,
+			ServerEndpointConfig->ConnCooldownSec,
+			ServerEndpointConfig->ConnCooldownMaxSec);
+	}
+
+	{
+		FScopeLock EndpointsLock(&EndpointsCS);
+
+		if (!Endpoints.IsEmpty())
+		{
+			for (auto It = Endpoints.CreateIterator(); It; ++It)
+			{
+				const TSharedPtr<FQuicEndpoint, ESPMode::ThreadSafe>& Endpoint = It->Value;
+
+				if (Endpoint.IsValid())
+				{
+					if (!Endpoint->IsStopping())
+					{
+						Endpoint->StopEndpoint();
+					}
+				}
+
+				It.RemoveCurrent();
+			}
+		}
+	}
+
+    EndpointMode = EEndpointMode::Server;
+
+	FString PrivateKey = ServerEndpointConfig->PrivateKey;
+	FString Certificate = ServerEndpointConfig->Certificate;
+
+	// Generate private key and certificate if none provided
+	if (PrivateKey == "" || Certificate == "")
+	{
+		if (!QuicCertificateUtils::CreateSelfSigned())
+		{
+			UE_LOG(LogQuicMessagingTransport, Error,
+				TEXT("[EndpointManager] Could not generate self-signed certificate."));
+
+			return;
+		}
+
+		TPair<FString, FString> SelfSigned = QuicCertificateUtils::GetSelfSignedPaths();
+		Certificate = SelfSigned.Key;
+		PrivateKey = SelfSigned.Value;
+
+		UE_LOG(LogQuicMessagingTransport, Verbose, 
+			TEXT("[EndpointManager] Generated self-signed certificate."));
+	}
+
+	CurrentConfiguration = LoadServerConfiguration(Certificate, PrivateKey);
+
+	if (!CurrentConfiguration)
+	{
+		UE_LOG(LogQuicMessagingTransport, Error,
+			TEXT("[EndpointManager] Server could not be started."));
+
+		return;
+	}
+
+    const TSharedPtr<FQuicEndpoint, ESPMode::ThreadSafe> LocalEndpoint
+        = MakeShared<FQuicServer, ESPMode::ThreadSafe>(
+		ServerEndpointConfig, CurrentConfiguration,
+		MsQuic, Registration, TEXT("QuicServerEndpoint"),Alpn);
+
+    LocalEndpoint->OnMessageInbound().BindRaw(this, &FQuicEndpointManager::ValidateInboundMessage);
+    LocalEndpoint->OnSerializeHeader().BindRaw(this, &FQuicEndpointManager::SerializeHeader);
+    LocalEndpoint->OnDeserializeHeader().BindRaw(this, &FQuicEndpointManager::DeserializeHeader);
+    LocalEndpoint->OnEndpointDisconnected().BindRaw(this, &FQuicEndpointManager::EndpointDisconnected);
+	LocalEndpoint->OnEndpointConnected().BindRaw(this, &FQuicEndpointManager::EndpointConnected);
+	LocalEndpoint->OnStatisticsUpdated().BindRaw(this, &FQuicEndpointManager::EndpointStatisticsUpdated);
+
+    LocalEndpoint->Start();
+
+	{
+		FScopeLock EndpointsLock(&EndpointsCS);
+		Endpoints.Add(ServerEndpointConfig->Endpoint, LocalEndpoint);
+	}
+
+	UE_LOG(LogQuicMessagingTransport, Verbose,
+		TEXT("[EndpointManager] Configured endpoint manager to run in server mode."));
+}
+
+
+void FQuicEndpointManager::AddClient(const FIPv4Endpoint& ClientEndpoint)
+{
+    if (EndpointMode == EEndpointMode::Server)
+    {
+        return;
+    }
+
+    {
+		FScopeLock EndpointsLock(&EndpointsCS);
+
+		if (const FQuicEndpointPtr* Endpoint = Endpoints.Find(ClientEndpoint))
+		{
+			UE_LOG(LogQuicMessagingTransport, Verbose,
+				TEXT("[EndpointManager] AddClient: endpoint %s still connected, removing it."),
+				*ClientEndpoint.ToString());
+
+			if (!(*Endpoint)->IsStopping())
+			{
+				(*Endpoint)->StopEndpoint();
+			}
+
+			Endpoints.Remove(ClientEndpoint);
+		}
+    }
+
+	const TSharedRef<FQuicClientConfig> ClientEndpointConfig
+		= StaticCastSharedRef<FQuicClientConfig, FQuicEndpointConfig>(EndpointConfig);
+
+	ClientEndpointConfig->Endpoint = ClientEndpoint;
+
+	if (!CurrentConfiguration)
+	{
+		CurrentConfiguration = LoadClientConfiguration();
+
+		if (!CurrentConfiguration)
+		{
+			UE_LOG(LogQuicMessagingTransport, Error,
+				TEXT("[EndpointManager] Client could not be started."));
+
+			return;
+		}
+	}
+
+    const TSharedPtr<FQuicEndpoint, ESPMode::ThreadSafe> RemoteEndpoint
+        = MakeShared<FQuicClient, ESPMode::ThreadSafe>(
+        ClientEndpointConfig, CurrentConfiguration,
+		MsQuic, Registration, TEXT("QuicClientEndpoint"));
+    
+    RemoteEndpoint->OnMessageInbound().BindRaw(this, &FQuicEndpointManager::ValidateInboundMessage);
+	RemoteEndpoint->OnSerializeHeader().BindRaw(this, &FQuicEndpointManager::SerializeHeader);
+	RemoteEndpoint->OnDeserializeHeader().BindRaw(this, &FQuicEndpointManager::DeserializeHeader);
+	RemoteEndpoint->OnEndpointDisconnected().BindRaw(this, &FQuicEndpointManager::EndpointDisconnected);
+	RemoteEndpoint->OnStatisticsUpdated().BindRaw(this, &FQuicEndpointManager::EndpointStatisticsUpdated);
+
+	RemoteEndpoint->Start();
+
+    {
+		FScopeLock EndpointsLock(&EndpointsCS);
+		Endpoints.Add(ClientEndpoint, RemoteEndpoint);
+    }
+
+	UE_LOG(LogQuicMessagingTransport, Verbose,
+		TEXT("[EndpointManager] Added client endpoint %s."),
+		*ClientEndpoint.ToString());
+}
+
+
+void FQuicEndpointManager::RemoveClient(const FIPv4Endpoint& ClientEndpoint)
+{
+	{
+		FScopeLock EndpointsLock(&EndpointsCS);
+
+		if (const FQuicEndpointPtr* Endpoint = Endpoints.Find(ClientEndpoint))
+		{
+			if (!(*Endpoint)->IsStopping())
+			{
+				(*Endpoint)->StopEndpoint();
+			}
+
+			Endpoints.Remove(ClientEndpoint);
+
+			UE_LOG(LogQuicMessagingTransport, Verbose,
+				TEXT("[EndpointManager] Removed client endpoint %s."),
+				*ClientEndpoint.ToString());
+		}
+	}
+
+	CheckLoseNode(ClientEndpoint);
+}
+
+
+/*
+ * TODO(vri): Fix reconnect mechanism
+ * Tracked in [UCS-5152] Finalize QuicMessaging plugin
+ */
+void FQuicEndpointManager::ReconnectClient(const FIPv4Endpoint ClientEndpoint)
+{
+	UE_LOG(LogQuicMessagingTransport, Verbose,
+		TEXT("[EndpointManager] Reconnecting client %s."),
+		*ClientEndpoint.ToString());
+
+	if (!ReconnectAttempts.Contains(ClientEndpoint))
+	{
+		ReconnectAttempts.Add(ClientEndpoint, 0);
+	}
+
+	if (ReconnectAttempts[ClientEndpoint] == 3)
+	{
+		return;
+	}
+
+	ReconnectAttempts[ClientEndpoint]++;
+
+	AddClient(ClientEndpoint);
+}
+
+
+bool FQuicEndpointManager::IsEndpointAuthenticated(const FGuid& NodeId) const
+{
+	FScopeLock NodesLock(&NodesCS);
+
+	const FQuicNodeInfo* SearchNode = KnownNodes.Find(NodeId);
+
+	if (!SearchNode)
+	{
+		return false;
+	}
+
+    return SearchNode->bIsAuthenticated;
+}
+
+
+void FQuicEndpointManager::SetEndpointAuthenticated(const FGuid& NodeId)
+{
+	FScopeLock NodesLock(&NodesCS);
+
+	FQuicNodeInfo* SearchNode = KnownNodes.Find(NodeId);
+
+	if (!SearchNode)
+	{
+		UE_LOG(LogQuicMessagingTransport, Error,
+			TEXT("[EndpointManager] Could not find node to set as authenticated: %s"),
+			*NodeId.ToString());
+
+		return;
+	}
+
+	SearchNode->bIsAuthenticated = true;
+
+    // Mark the handler as authenticated as well
+    if (EndpointMode == EEndpointMode::Server)
+    {
+		FScopeLock EndpointsLock(&EndpointsCS);
+
+		const FQuicEndpointPtr* LocalEndpoint = Endpoints.Find(SearchNode->LocalEndpoint);
+
+		if (!LocalEndpoint)
+		{
+			UE_LOG(LogQuicMessagingTransport, Error,
+				TEXT("[EndpointManager] Could not find server endpoint to set handler as authenticated."));
+			
+			return;
+		}
+
+        (*LocalEndpoint)->SetHandlerAuthenticated(SearchNode->Endpoint);
+    }
+}
+
+
+void FQuicEndpointManager::DisconnectNode(const FGuid& NodeId)
+{
+	FScopeLock NodesLock(&NodesCS);
+
+	const FQuicNodeInfo* SearchNode = KnownNodes.Find(NodeId);
+
+	if (!SearchNode)
+	{
+		UE_LOG(LogQuicMessagingTransport, Error,
+			TEXT("[EndpointManager] Could not find node to disconnect: %s"),
+			*NodeId.ToString());
+
+		return;
+	}
+
+	DisconnectEndpoint(SearchNode->LocalEndpoint);
+}
+
+
+void FQuicEndpointManager::DisconnectEndpoint(const FIPv4Endpoint& Endpoint)
+{
+	// It's a client, so just remove it
+	if (EndpointMode == EEndpointMode::Client)
+	{
+		RemoveClient(Endpoint);
+		return;
+	}
+
+	FScopeLock EndpointsLock(&EndpointsCS);
+
+	const FQuicEndpointPtr* LocalEndpoint = Endpoints.Find(Endpoint);
+
+	if (!LocalEndpoint)
+	{
+		UE_LOG(LogQuicMessagingTransport, Error,
+			TEXT("[EndpointManager] Could not find server endpoint to disconnect handler."));
+
+		return;
+	}
+
+	(*LocalEndpoint)->DisconnectHandler(Endpoint);
+}
+
+
+void FQuicEndpointManager::EnqueueOutboundMessages(FQuicPayloadPtr Data,
+	const TArray<TTuple<FGuid, uint32>>& MessageMetas, const EQuicMessageType MessageType)
+{
+    if (!Data.IsValid())
+    {
+        return;
+    }
+
+    if (Data->IsEmpty())
+    {
+        return;
+    }
+
+	for (const TTuple<FGuid, uint32>& MessageMeta : MessageMetas)
+	{
+		const FGuid& Recipient = MessageMeta.Key;
+
+		FIPv4Endpoint LocalAddress;
+		FIPv4Endpoint RecipientAddress;
+
+		{
+			FScopeLock NodesLock(&NodesCS);
+
+			const FQuicNodeInfo* RecipientNode = KnownNodes.Find(Recipient);
+
+			if (!RecipientNode)
+			{
+				UE_LOG(LogQuicMessagingTransport, Error,
+					TEXT("[EndpointManager] Could not find message recipient node %s"),
+					*Recipient.ToString());
+
+				continue;
+			}
+
+			LocalAddress = RecipientNode->LocalEndpoint;
+			RecipientAddress = RecipientNode->Endpoint;
+		}
+
+		FMessageHeader MessageHeader(MessageType, Recipient, LocalNodeId, Data->Num());
+		const FQuicPayloadPtr MetaBuffer = SerializeHeaderDelegate.Execute(MessageHeader);
+		
+		const FOutboundMessage OutboundMessage(
+			RecipientAddress, Data, MetaBuffer);
+
+		const FIPv4Endpoint& SendingEndpointAddress
+			= (EndpointMode == EEndpointMode::Client) ? RecipientAddress : LocalAddress;
+
+		FScopeLock EndpointsLock(&EndpointsCS);
+
+		const FQuicEndpointPtr* SendingEndpoint = Endpoints.Find(SendingEndpointAddress);
+
+		if (!SendingEndpoint)
+		{
+			UE_LOG(LogQuicMessagingTransport, Error,
+				TEXT("[EndpointManager] Could not find sending endpoint to enqueue message on: %s"),
+				*SendingEndpointAddress.ToString());
+
+			continue;
+		}
+
+		(*SendingEndpoint)->EnqueueOutboundMessage(OutboundMessage);
+	}
+}
+
+
+void FQuicEndpointManager::ValidateInboundMessage(const FInboundMessage InboundMessage)
+{
+    // Message was generated locally
+    if (InboundMessage.MessageHeader.RecipientId == InboundMessage.MessageHeader.SenderId)
+    {
+        return;
+    }
+
+    // Sender guid is not initialized
+    if (InboundMessage.MessageHeader.SenderId == FGuid())
+    {
+        return;
+    }
+
+    // Incorrect message segment for server side
+    if (EndpointMode == EEndpointMode::Server
+        && InboundMessage.MessageHeader.MessageType == EQuicMessageType::AuthenticationResponse)
+    {
+        return;
+    }
+
+    // Incorrect message segment for client side
+    if (EndpointMode == EEndpointMode::Client
+        && InboundMessage.MessageHeader.MessageType == EQuicMessageType::Authentication)
+    {
+        return;
+    }
+
+    // Node discovery
+    if (!KnownNodes.Contains(InboundMessage.MessageHeader.SenderId))
+    {
+		// Sometimes the first transmitted message (EQuicMessageType::Hello) is broken
+		if (!IsMessageValid(InboundMessage))
+		{
+			SendEndpointHello(InboundMessage);
+			return;
+		}
+
+        EndpointNodeDiscovered(InboundMessage.MessageHeader.SenderId,
+            InboundMessage.Sender, InboundMessage.Receiver);
+
+		SendEndpointHello(InboundMessage);
+    }
+
+	if (InboundMessage.MessageHeader.MessageType == EQuicMessageType::Bye)
+	{
+		EndpointNodeLost(InboundMessage.MessageHeader.SenderId, InboundMessage.Sender);
+		EndpointDisconnected(InboundMessage.Sender, EQuicEndpointError::Normal);
+	}
+
+	// Don't have to pass these messages on
+	if (InboundMessage.MessageHeader.MessageType == EQuicMessageType::Hello
+		|| InboundMessage.MessageHeader.MessageType == EQuicMessageType::Bye)
+	{
+		return;
+	}
+
+    MessageValidatedDelegate.ExecuteIfBound(InboundMessage);
+}
+
+
+void FQuicEndpointManager::SendEndpointHello(const FInboundMessage& InboundMessage)
+{
+	FScopeLock EndpointsLock(&EndpointsCS);
+
+	if (EndpointMode == EEndpointMode::Client)
+	{
+		if (const FQuicEndpointPtr* Client = Endpoints.Find(InboundMessage.Sender))
+		{
+			(*Client)->SendHello(InboundMessage.Sender);
+		}
+	}
+}
+
+
+bool FQuicEndpointManager::IsMessageValid(const FInboundMessage InboundMessage)
+{
+	if (InboundMessage.MessageHeader.MessageType < EQuicMessageType::Hello
+		|| InboundMessage.MessageHeader.MessageType > EQuicMessageType::Bye)
+	{
+		return false;
+	}
+
+	return true;
+}
+
+
+FQuicPayloadPtr FQuicEndpointManager::SerializeHeader(
+	FMessageHeader& MessageHeader) const
+{
+	return SerializeHeaderDelegate.Execute(MessageHeader);
+}
+
+
+FMessageHeader FQuicEndpointManager::DeserializeHeader(
+	const FQuicPayloadPtr HeaderData) const
+{
+	return DeserializeHeaderDelegate.Execute(HeaderData);
+}
+
+
+void FQuicEndpointManager::EndpointNodeDiscovered(const FGuid NodeId,
+    const FIPv4Endpoint& NodeEndpoint, const FIPv4Endpoint& LocalEndpoint)
+{
+	{
+		FScopeLock NodesLock(&NodesCS);
+
+		if (FQuicNodeInfo* ExistingNode = KnownNodes.Find(NodeId))
+		{
+			UE_LOG(LogQuicMessagingTransport, Display,
+				TEXT("[EndpointManager] Discovered node already known."));
+
+			if (ExistingNode->Endpoint != NodeEndpoint)
+			{
+				ExistingNode->Endpoint = NodeEndpoint;
+			}
+
+			return;
+		}
+
+		const FQuicNodeInfo NodeInfo(NodeEndpoint, NodeId, LocalEndpoint);
+
+		KnownNodes.Add(NodeId, NodeInfo);
+	}
+
+    EndpointNodeDiscoveredDelegate.ExecuteIfBound(NodeId);
+}
+
+
+void FQuicEndpointManager::EndpointNodeLost(const FGuid NodeId, const FIPv4Endpoint LostEndpoint)
+{
+	{
+		FScopeLock NodesLock(&NodesCS);
+
+		const FQuicNodeInfo* ExistingNode = KnownNodes.Find(NodeId);
+
+		if (!ExistingNode)
+		{
+			UE_LOG(LogQuicMessagingTransport, Warning,
+				TEXT("[EndpointManager] Node to forget not found."));
+
+			return;
+		}
+
+		if (ExistingNode->Endpoint != LostEndpoint)
+		{
+			return;
+		}
+
+		KnownNodes.Remove(NodeId);
+	}
+
+    EndpointNodeLostDelegate.ExecuteIfBound(NodeId);
+}
+
+
+TOptional<FGuid> FQuicEndpointManager::GetKnownNodeId(const FIPv4Endpoint& Endpoint) const
+{
+	FScopeLock NodesLock(&NodesCS);
+
+	for (const auto& NodePair : KnownNodes)
+	{
+		if (NodePair.Value.Endpoint == Endpoint)
+		{
+			return TOptional<FGuid>(NodePair.Key);
+		}
+	}
+
+	return TOptional<FGuid>();
+}
+
+
+bool FQuicEndpointManager::IsEndpointKnown(const FIPv4Endpoint Endpoint) const
+{
+	FScopeLock NodesLock(&NodesCS);
+
+	bool EndpointKnown = false;
+
+	for (const auto& NodePair : KnownNodes)
+	{
+		if (NodePair.Value.Endpoint == Endpoint)
+		{
+			EndpointKnown = true;
+			break;
+		}
+	}
+
+	return EndpointKnown;
+}
+
+
+TArray<FGuid> FQuicEndpointManager::GetKnownNodeIds() const
+{
+	FScopeLock NodesLock(&NodesCS);
+
+	TArray<FGuid> KnownNodeIds;
+	KnownNodes.GetKeys(KnownNodeIds);
+
+	return KnownNodeIds;
+}
+
+
+TArray<FIPv4Endpoint> FQuicEndpointManager::GetKnownEndpoints() const
+{
+	FScopeLock NodesLock(&NodesCS);
+
+	TArray<FIPv4Endpoint> KnownEndpoints;
+
+	for (const auto& NodePair : KnownNodes)
+	{
+		KnownEndpoints.Add(NodePair.Value.Endpoint);
+	}
+
+	return KnownEndpoints;
+}
+
+
+FMessageTransportStatistics FQuicEndpointManager::GetStats(FGuid NodeId) const
+{
+	FScopeLock NodesLock(&NodesCS);
+
+	if (FQuicNodeInfo const* NodeInfo = KnownNodes.Find(NodeId))
+	{
+		return NodeInfo->Statistics;
+	}
+
+	return {};
+}
+
+
+void FQuicEndpointManager::EndpointDisconnected(
+	const FIPv4Endpoint DisconnectedEndpoint, EQuicEndpointError Reason)
+{
+	UE_LOG(LogQuicMessagingTransport, Verbose,
+		TEXT("[EndpointManager] Endpoint disconnected: %s / Reason: %d."),
+		*DisconnectedEndpoint.ToString(), Reason);
+
+	{
+		FScopeLock EndpointsLock(&EndpointsCS);
+
+		// Shutdown disconnected client
+		if (const FQuicEndpointPtr* Endpoint = Endpoints.Find(DisconnectedEndpoint))
+		{
+			if (!(*Endpoint)->IsStopping())
+			{
+				(*Endpoint)->StopEndpoint();
+				Endpoints.Remove(DisconnectedEndpoint);
+			}
+		}
+	}
+
+	CheckLoseNode(DisconnectedEndpoint);
+
+	// Attempt reconnect for aborted client side connections
+	if (EndpointMode == EEndpointMode::Client && Reason == EQuicEndpointError::ConnectionAbort)
+	{
+		ReconnectClient(DisconnectedEndpoint);
+	}
+}
+
+
+void FQuicEndpointManager::CheckLoseNode(const FIPv4Endpoint LostEndpoint)
+{
+	FScopeLock NodesLock(&NodesCS);
+
+	// Search for and mark node as lost
+	for (const auto& NodePair : KnownNodes)
+	{
+		if (NodePair.Value.Endpoint == LostEndpoint)
+		{
+			NodesLock.Unlock();
+
+			EndpointNodeLost(NodePair.Value.NodeId, LostEndpoint);
+			break;
+		}
+	}
+}
+
+
+void FQuicEndpointManager::EndpointConnected(const FIPv4Endpoint& ConnectedEndpoint)
+{
+	if (EndpointConfig->DiscoveryTimeoutSec == 0)
+	{
+		return;
+	}
+
+	FScopeLock DiscoveryTimeoutsLock(&DiscoveryTimeoutCS);
+
+	FDiscoveryTimeoutEntry TimeoutEntry;
+
+	TimeoutEntry.Endpoint = ConnectedEndpoint;
+	TimeoutEntry.Timestamp = FPlatformTime::Seconds() + EndpointConfig->DiscoveryTimeoutSec;
+
+	DiscoveryTimeouts.Add(TimeoutEntry);
+}
+
+
+void FQuicEndpointManager::CheckDiscoveryTimeout()
+{
+	FScopeLock DiscoveryTimeoutsLock(&DiscoveryTimeoutCS);
+
+	if (DiscoveryTimeouts.IsEmpty())
+	{
+		return;
+	}
+
+	const double Now = FPlatformTime::Seconds();
+
+	for (auto Entry = DiscoveryTimeouts.CreateIterator(); Entry; ++Entry)
+	{
+		if (Entry->Timestamp < Now)
+		{
+			return;
+		}
+
+		if (!IsEndpointKnown(Entry->Endpoint))
+		{
+			DisconnectEndpoint(Entry->Endpoint);
+		}
+
+		Entry.RemoveCurrent();
+	}
+}
+
+
+void FQuicEndpointManager::EndpointStatisticsUpdated(
+	const FMessageTransportStatistics TransportStats)
+{
+	FScopeLock NodesLock(&NodesCS);
+
+	for (auto& NodePair : KnownNodes)
+	{
+		if (NodePair.Value.Endpoint.ToString() == TransportStats.IPv4AsString)
+		{
+			NodePair.Value.Statistics = TransportStats;
+		}
+	}
+}
+
diff --git a/Engine/Plugins/Experimental/QuicMessaging/Source/QuicMessagingTransport/Private/QuicIncludes.h b/Engine/Plugins/Experimental/QuicMessaging/Source/QuicMessagingTransport/Private/QuicIncludes.h
new file mode 100644
--- /dev/null
+++ b/Engine/Plugins/Experimental/QuicMessaging/Source/QuicMessagingTransport/Private/QuicIncludes.h
@@ -0,0 +1,54 @@
+// Copyright Epic Games, Inc. All Rights Reserved.
+
+#pragma once
+
+
+#if PLATFORM_WINDOWS
+#include "Windows/AllowWindowsPlatformTypes.h"
+THIRD_PARTY_INCLUDES_START
+#endif
+
+// Inspired by
+// Engine\Plugins\Media\PixelStreaming\Source\PixelStreaming\Private\WebRTCIncludes.h
+
+// C5105: One of the files included by "msquic.h" has a macro expansion
+//		  producing defined, which has undefined behavior
+#pragma warning(push)
+#pragma warning(disable: 5105)
+
+#ifdef __clang__
+#pragma clang diagnostic push
+#pragma clang diagnostic ignored "-Wmissing-braces"
+#endif
+
+
+// This define is needed to provide optional local non-encrypted traffic
+#define QUIC_API_ENABLE_INSECURE_FEATURES
+
+#include <msquic.h>
+#include <stdio.h>
+#include <stdlib.h>
+
+
+#ifdef __clang__
+#pragma clang diagnostic pop
+#endif
+
+#pragma warning(pop)
+
+
+#ifndef UNREFERENCED_PARAMETER
+#define UNREFERENCED_PARAMETER(P) (void)(P)
+#endif
+
+#define UI UI_ST
+#include <openssl/pem.h>
+#include <openssl/x509.h>
+#undef UI
+
+
+#if PLATFORM_WINDOWS
+THIRD_PARTY_INCLUDES_END
+#include "Windows/HideWindowsPlatformTypes.h"
+#endif
+
diff --git a/Engine/Plugins/Experimental/QuicMessaging/Source/QuicMessagingTransport/Private/QuicMessagingTransportModule.cpp b/Engine/Plugins/Experimental/QuicMessaging/Source/QuicMessagingTransport/Private/QuicMessagingTransportModule.cpp
new file mode 100644
--- /dev/null
+++ b/Engine/Plugins/Experimental/QuicMessaging/Source/QuicMessagingTransport/Private/QuicMessagingTransportModule.cpp
@@ -0,0 +1,21 @@
+// Copyright Epic Games, Inc. All Rights Reserved.
+
+#include "IQuicMessagingTransportModule.h"
+
+class FQuicMessagingTransportModule : public IQuicMessagingTransportModule
+{
+public:
+	FQuicMessagingTransportModule() = default;
+	virtual ~FQuicMessagingTransportModule() {}
+
+	virtual void StartupModule() override
+	{
+	}
+
+	virtual void ShutdownModule() override
+	{
+	}
+};
+
+
+IMPLEMENT_MODULE(FQuicMessagingTransportModule, QuicMessagingTransport);
diff --git a/Engine/Plugins/Experimental/QuicMessaging/Source/QuicMessagingTransport/Private/QuicMessagingTransportPrivate.cpp b/Engine/Plugins/Experimental/QuicMessaging/Source/QuicMessagingTransport/Private/QuicMessagingTransportPrivate.cpp
new file mode 100644
--- /dev/null
+++ b/Engine/Plugins/Experimental/QuicMessaging/Source/QuicMessagingTransport/Private/QuicMessagingTransportPrivate.cpp
@@ -0,0 +1,6 @@
+// Copyright Epic Games, Inc. All Rights Reserved.
+
+#include "QuicMessagingTransportPrivate.h"
+
+DEFINE_LOG_CATEGORY(LogQuicMessagingTransport);
+
diff --git a/Engine/Plugins/Experimental/QuicMessaging/Source/QuicMessagingTransport/Private/QuicMessagingTransportPrivate.h b/Engine/Plugins/Experimental/QuicMessaging/Source/QuicMessagingTransport/Private/QuicMessagingTransportPrivate.h
new file mode 100644
--- /dev/null
+++ b/Engine/Plugins/Experimental/QuicMessaging/Source/QuicMessagingTransport/Private/QuicMessagingTransportPrivate.h
@@ -0,0 +1,8 @@
+// Copyright Epic Games, Inc. All Rights Reserved.
+
+#pragma once
+
+#include "Logging/LogMacros.h"
+
+DECLARE_LOG_CATEGORY_EXTERN(LogQuicMessagingTransport, Log, All);
+
diff --git a/Engine/Plugins/Experimental/QuicMessaging/Source/QuicMessagingTransport/Private/QuicQueues.cpp b/Engine/Plugins/Experimental/QuicMessaging/Source/QuicMessagingTransport/Private/QuicQueues.cpp
new file mode 100644
--- /dev/null
+++ b/Engine/Plugins/Experimental/QuicMessaging/Source/QuicMessagingTransport/Private/QuicQueues.cpp
@@ -0,0 +1,20 @@
+// Copyright Epic Games, Inc. All Rights Reserved.
+
+#include "QuicQueues.h"
+
+
+void FQuicQueues::ConsumeMessages()
+{
+	const double CurrentTime = FPlatformTime::Seconds();
+
+	if (CurrentTime < (LastConsumed + ConsumeInterval.GetTotalSeconds()))
+	{
+		return;
+	}
+
+	ConsumeInboundMessages();
+	ProcessInboundBuffers();
+	ConsumeOutboundMessages();
+
+	LastConsumed = FPlatformTime::Seconds();
+}
diff --git a/Engine/Plugins/Experimental/QuicMessaging/Source/QuicMessagingTransport/Private/QuicQueues.h b/Engine/Plugins/Experimental/QuicMessaging/Source/QuicMessagingTransport/Private/QuicQueues.h
new file mode 100644
--- /dev/null
+++ b/Engine/Plugins/Experimental/QuicMessaging/Source/QuicMessagingTransport/Private/QuicQueues.h
@@ -0,0 +1,71 @@
+// Copyright Epic Games, Inc. All Rights Reserved.
+
+#pragma once
+
+#include "CoreMinimal.h"
+#include "HAL/PlatformTime.h"
+#include "Containers/Queue.h"
+
+#include "QuicMessages.h"
+
+
+class FQuicQueues
+{
+
+protected:
+
+	FQuicQueues() = default;
+
+	virtual ~FQuicQueues() = default;
+
+public:
+
+	/**
+	 * Periodically called to consume local inbound and outbound messages.
+	 * Consumes messages every ConsumeInterval microseconds.
+	 */
+	void ConsumeMessages();
+
+protected:
+
+    /**
+	 * Dequeue inbound messages until MaxConsumeTime is reached.
+	 * Inbound messages are passed to QuicEndpointManager for validation.
+	 */
+	virtual void ConsumeInboundMessages() = 0;
+
+	/**
+	 * Dequeue inbound buffers until MaxConsumeTime is reached.
+	 * Inbound buffers via QUIC are copied into our buffers until the stream end is indicated.
+	 */
+	virtual void ProcessInboundBuffers() = 0;
+
+	/**
+	 * Dequeue outbound messages until MaxConsumeTime is reached.
+	 * Outbound messages will then be sent via QUIC.
+	 */
+	virtual void ConsumeOutboundMessages() = 0;
+
+protected:
+
+    /** Queue for inbound messages ready to be consumed. */
+    TQueue<FInboundMessage> InboundMessages;
+
+	/** Queue for outbound message ready to be sent. */
+	TQueue<FOutboundMessage> OutboundMessages;
+
+protected:
+
+    /** Holds current message consume time in microseconds. */
+    const FTimespan MaxConsumeTime = 5000;
+
+	/** Holds the interval for message consume. */
+	const FTimespan ConsumeInterval = 10000;
+
+    /** Holds timestamp of when messages have last been consumed. */
+    double LastConsumed = FPlatformTime::Seconds();
+
+    /** Current id for InboundMessages. */
+    uint32 InboundMessagesId = 0;
+
+};
diff --git a/Engine/Plugins/Experimental/QuicMessaging/Source/QuicMessagingTransport/Private/QuicServer.cpp b/Engine/Plugins/Experimental/QuicMessaging/Source/QuicMessagingTransport/Private/QuicServer.cpp
new file mode 100644
--- /dev/null
+++ b/Engine/Plugins/Experimental/QuicMessaging/Source/QuicMessagingTransport/Private/QuicServer.cpp
@@ -0,0 +1,1170 @@
+// Copyright Epic Games, Inc. All Rights Reserved.
+
+#include "QuicServer.h"
+
+#include "QuicEndpointManager.h"
+
+
+void FQuicServer::StartEndpoint()
+{
+	QUIC_STATUS Status;
+
+	const QUIC_ADDR Address = GetAddressFromEndpoint(Endpoint);
+
+	/**
+	 * Allocate new listener object
+	 */
+	if (QUIC_FAILED(Status = MsQuic->ListenerOpen(Registration,
+		ServerListenerCallback, this, &ServerListener)))
+	{
+		const EQuicEndpointError ListenerError = QuicUtils::ConvertQuicStatus(Status);
+
+		UE_LOG(LogQuicMessagingTransport, Error,
+			TEXT("[QuicServer] Could not open listener: %s."),
+			*QuicUtils::GetEndpointErrorString(ListenerError));
+
+		return;
+	}
+
+	/**
+	 * Start listening for incoming connections
+	 */
+	if (QUIC_FAILED(Status = MsQuic->ListenerStart(ServerListener, &Alpn, 1, &Address)))
+	{
+		const EQuicEndpointError ListenerError = QuicUtils::ConvertQuicStatus(Status);
+
+		UE_LOG(LogQuicMessagingTransport, Error,
+			TEXT("[QuicServer] Could not start listener: %s."),
+			*QuicUtils::GetEndpointErrorString(ListenerError));
+
+		return;
+	}
+
+	ConnectionState = EQuicServerState::Running;
+
+	// Setup the connection cooldown check tick
+	ConnectionCooldownCheckHandle = FTSTicker::GetCoreTicker().AddTicker(
+		TEXT("QuicServerConnCooldownTick"), 1, [this](float DeltaSeconds) {
+
+		CheckConnectionCooldown();
+		return true;
+
+	});
+
+	UE_LOG(LogQuicMessagingTransport, Display,
+		TEXT("[QuicServer] Started server. Listening on %s."),
+		*GetEndpointFromAddress(Address).ToString());
+}
+
+
+void FQuicServer::StopEndpoint()
+{
+	if (IsStopping())
+	{
+		return;
+	}
+
+	UE_LOG(LogQuicMessagingTransport, Verbose, 
+		TEXT("[QuicServer] Stopping server ..."));
+
+	Stop();
+
+	if (ConnectionCooldownCheckHandle.IsValid())
+	{
+		FTSTicker::GetCoreTicker().RemoveTicker(ConnectionCooldownCheckHandle);
+	}
+
+	if (ConnectionState < EQuicServerState::Failed)
+	{
+		ConnectionState = EQuicServerState::Stopped;
+	}
+
+	if (MsQuic)
+	{
+		if (ServerListener)
+		{
+			MsQuic->ListenerStop(ServerListener);
+			MsQuic->ListenerClose(ServerListener);
+		}
+
+		ShutdownAllHandlers();
+	}
+
+	UE_LOG(LogQuicMessagingTransport, Verbose, 
+		TEXT("[QuicServer] Stopped server."));
+}
+
+
+bool FQuicServer::IsEndpointReady()
+{
+	return ConnectionState == EQuicServerState::Running;
+}
+
+
+void FQuicServer::UpdateEndpoint()
+{
+	if (IsStopping() || ConnectionState != EQuicServerState::Running)
+	{
+		return;
+	}
+
+	if (Handlers.IsEmpty())
+	{
+		return;
+	}
+
+	ConsumeMessages();
+}
+
+
+void FQuicServer::ConsumeInboundMessages()
+{
+	if (InboundMessages.IsEmpty())
+	{
+		return;
+	}
+
+	double CurrentTime = FPlatformTime::Seconds();
+	const double ConsumeLimit = CurrentTime + MaxConsumeTime.GetTotalSeconds();
+
+	FInboundMessage InboundMessage;
+
+	while (CurrentTime < ConsumeLimit && InboundMessages.Dequeue(InboundMessage))
+	{
+		MessageInboundDelegate.ExecuteIfBound(InboundMessage);
+
+		CurrentTime = FPlatformTime::Seconds();
+	}
+}
+
+
+void FQuicServer::ProcessInboundBuffers()
+{
+	if (InboundQuicBuffers.IsEmpty())
+	{
+		return;
+	}
+
+	double CurrentTime = FPlatformTime::Seconds();
+	const double ConsumeLimit = CurrentTime + MaxConsumeTime.GetTotalSeconds();
+
+	FInboundQuicBuffer InboundBuffer;
+
+	while (CurrentTime < ConsumeLimit && InboundQuicBuffers.Dequeue(InboundBuffer))
+	{
+		if (ReceiveFromStream(InboundBuffer))
+		{
+			if (InboundBuffer.bEndOfStream)
+			{
+				ReachedEndOfStream(InboundBuffer);
+			}
+		}
+
+		CurrentTime = FPlatformTime::Seconds();
+	}
+}
+
+
+void FQuicServer::ConsumeOutboundMessages()
+{
+	if (OutboundMessages.IsEmpty())
+	{
+		return;
+	}
+
+	double CurrentTime = FPlatformTime::Seconds();
+	const double ConsumeLimit = CurrentTime + MaxConsumeTime.GetTotalSeconds();
+	
+	FOutboundMessage OutboundMessage;
+
+	while (CurrentTime < ConsumeLimit && OutboundMessages.Dequeue(OutboundMessage))
+	{
+		HQUIC HandlerConnection;
+
+		{
+			FScopeLock HandlersLock(&HandlersCS);
+
+			const FQuicServerHandlerPtr* ServerHandler
+				= Handlers.Find(OutboundMessage.Recipient);
+
+			if (!ServerHandler)
+			{
+				UE_LOG(LogQuicMessagingTransport, Error,
+					TEXT("[QuicServer] Could not find server handler when preparing outbound message: %s"),
+					*OutboundMessage.Recipient.ToString());
+
+				continue;
+			}
+
+			if ((*ServerHandler)->IsStopping())
+			{
+				continue;
+			}
+
+			if (!(*ServerHandler)->IsRunning())
+			{
+				OutboundMessages.Enqueue(OutboundMessage);
+				continue;
+			}
+
+			HandlerConnection = (*ServerHandler)->Connection;
+		}
+
+		SendToStream(HandlerConnection, OutboundMessage);
+
+		CurrentTime = FPlatformTime::Seconds();
+	}
+}
+
+
+bool FQuicServer::CanHandlerConnect(const FIPv4Endpoint& ConnectingEndpoint) const
+{
+	const FIPv4Endpoint ClearedEndpoint = ClearEndpointPort(ConnectingEndpoint);
+
+	FScopeLock ConnCooldownLock(&ConnectionCooldownCS);
+
+	return !ConnectionCooldown.Contains(ClearedEndpoint);
+}
+
+
+void FQuicServer::AddConnectionAttempt(const FIPv4Endpoint& ConnectingEndpoint)
+{
+	if (ServerConfig->ConnCooldownMode != EConnectionCooldownMode::Enabled)
+	{
+		return;
+	}
+
+	const FIPv4Endpoint ClearedEndpoint = ClearEndpointPort(ConnectingEndpoint);
+
+	FScopeLock ConnCooldownLock(&ConnectionCooldownCS);
+
+	const uint64 Now = FMath::Floor(FPlatformTime::Seconds());
+
+	FConnectionAttempts* ExistingAttempts = ConnectionAttempts.Find(ClearedEndpoint);
+
+	if (!ExistingAttempts)
+	{
+		FConnectionAttempts Attempts;
+		Attempts.AttemptTimestamps.Add(Now);
+
+		ConnectionAttempts.Add(ClearedEndpoint, Attempts);
+
+		return;
+	}
+
+	ExistingAttempts->AttemptTimestamps.Add(Now);
+
+	const int32 MaximumAttempts = ServerConfig->ConnCooldownMaxAttempts;
+
+	// Prune unnecessary previous attempts
+	if (ExistingAttempts->AttemptTimestamps.Num() > MaximumAttempts)
+	{
+		const int32 RemoveNumAttempts
+			= FMath::Abs(MaximumAttempts - ExistingAttempts->AttemptTimestamps.Num());
+
+		ExistingAttempts->AttemptTimestamps.RemoveAt(0, RemoveNumAttempts);
+	}
+
+	// Not enough attempts
+	if (ExistingAttempts->AttemptTimestamps.Num() < MaximumAttempts)
+	{
+		return;
+	}
+
+	const uint64 EarliestAttemptTimestamp = ExistingAttempts->AttemptTimestamps[0];
+
+	// Earliest attempt is longer ago
+	if (Now > (EarliestAttemptTimestamp + ServerConfig->ConnCooldownPeriodSec))
+	{
+		return;
+	}
+
+	// Calculate cooldown multiplier
+	const uint64 CooldownMultiplier = FMath::Floor(
+		FMath::Exp2(static_cast<float>(ExistingAttempts->NumCooldowns)));
+
+	// Get the cooldown seconds
+	uint64 CooldownSeconds = ServerConfig->ConnCooldownSec * CooldownMultiplier;
+
+	// Make sure it doesn't exceed the maximum cooldown
+	CooldownSeconds = FMath::Min(CooldownSeconds,
+		static_cast<uint64>(ServerConfig->ConnCooldownMaxSec));
+
+	const uint64 CooldownExpires = Now + (CooldownSeconds * CooldownMultiplier);
+
+	ExistingAttempts->NumCooldowns++;
+	ExistingAttempts->AttemptTimestamps.Empty();
+
+	ConnectionCooldown.Add(ClearedEndpoint, CooldownExpires);
+
+	UE_LOG(LogQuicMessagingTransport, Warning,
+		TEXT("Connection cooldown of %d seconds for IP %s."),
+		CooldownSeconds, *ClearedEndpoint.ToInternetAddr()->ToString(false));
+}
+
+
+void FQuicServer::CheckConnectionCooldown()
+{
+	FScopeLock ConnCooldownLock(&ConnectionCooldownCS);
+
+	if (ConnectionCooldown.IsEmpty())
+	{
+		return;
+	}
+
+	if (ServerConfig->ConnCooldownMode == EConnectionCooldownMode::Disabled)
+	{
+		ConnectionCooldown.Empty();
+
+		return;
+	}
+
+	const uint64 Now = FMath::Floor(FPlatformTime::Seconds());
+
+	for (auto CooldownIt = ConnectionCooldown.CreateIterator(); CooldownIt; ++CooldownIt)
+	{
+		if (Now > CooldownIt->Value)
+		{
+			CooldownIt.RemoveCurrent();
+		}
+	}
+}
+
+
+/*
+ * TODO(vri): Fully implement and fix server statistics
+ * Tracked in [UCS-5152] Finalize QuicMessaging plugin
+ */
+void FQuicServer::CollectStatistics()
+{
+	if (IsStopping() || ConnectionState != EQuicServerState::Running)
+	{
+		return;
+	}
+
+	if (MsQuic == nullptr)
+	{
+		return;
+	}
+
+	FScopeLock HandlersLock(&HandlersCS);
+
+	TArray<FIPv4Endpoint> UpdateHandlers;
+	Handlers.GetKeys(UpdateHandlers);
+
+	for (FIPv4Endpoint& UpdateEndpoint : UpdateHandlers)
+	{
+		FQuicServerHandlerPtr Handler = Handlers[UpdateEndpoint];
+
+		if (!Handler.IsValid() || !Handler->IsRunning())
+		{
+			continue;
+		}
+
+		QUIC_STATISTICS Stats{};
+
+		const QUIC_STATUS Status = GetConnectionStats(Handler->Connection, Stats);
+
+		if (Status != QUIC_STATUS_SUCCESS)
+		{
+			UE_LOG(LogQuicMessagingTransport, Error,
+				TEXT("[QuicServer] Could not collect statistics: %s."),
+				*QuicUtils::ConvertResult(Status));
+
+			return;
+		}
+
+		if (Handler->NeedRttReset())
+		{
+			Handler->RttMeasurements.Empty();
+			Handler->RttStart = FPlatformTime::Seconds();
+		}
+
+		Handler->RttMeasurements.Add(Stats.Rtt);
+
+		FMessageTransportStatistics TransportStats = ConvertStatistics(Handler->Endpoint, Stats);
+		TransportStats.AverageRTT = FTimespan::FromMicroseconds(Handler->GetAverageRtt());
+
+		StatisticsUpdatedDelegate.ExecuteIfBound(TransportStats);
+	}
+}
+
+
+void FQuicServer::SetHandlerAuthenticated(const FIPv4Endpoint& HandlerEndpoint)
+{
+	FScopeLock HandlersLock(&HandlersCS);
+
+	if (!Handlers.Contains(HandlerEndpoint))
+	{
+		return;
+	}
+
+	Handlers[HandlerEndpoint]->bAuthenticated = true;
+}
+
+
+void FQuicServer::DisconnectHandler(const FIPv4Endpoint& HandlerEndpoint)
+{
+	FScopeLock HandlersLock(&HandlersCS);
+
+	const FQuicServerHandlerPtr* Handler = Handlers.Find(HandlerEndpoint);
+
+	if (!Handler)
+	{
+		UE_LOG(LogQuicMessagingTransport, Error,
+			TEXT("[QuicServer] Could not find handler to disconnect: %s."),
+			*HandlerEndpoint.ToString());
+
+		return;
+	}
+
+	ShutdownHandler((*Handler)->Connection);
+}
+
+
+void FQuicServer::ShutdownAllHandlers() const
+{
+	FScopeLock HandlersLock(&HandlersCS);
+
+	for (const auto& Handler : Handlers)
+	{
+		if (Handler.Value->Connection)
+		{
+			ShutdownHandler(Handler.Value->Connection);
+		}
+	}
+}
+
+
+void FQuicServer::ShutdownHandler(const HQUIC HandlerConnection) const
+{
+	if (MsQuic)
+	{
+		MsQuic->ConnectionShutdown(
+			HandlerConnection, QUIC_CONNECTION_SHUTDOWN_FLAG_NONE, 0);
+	}
+}
+
+
+void FQuicServer::EnqueueOutboundMessage(const FOutboundMessage& OutboundMessage)
+{
+	OutboundMessages.Enqueue(OutboundMessage);
+}
+
+
+void FQuicServer::SendHello(FIPv4Endpoint RemoteEndpoint)
+{
+	if (IsStopping() || ConnectionState != EQuicServerState::Running)
+	{
+		return;
+	}
+
+	FMessageHeader MessageHeader(EQuicMessageType::Hello, FGuid::NewGuid(), LocalNodeId, 0);
+	const FQuicPayloadPtr MetaBuffer = SerializeHeaderDelegate.Execute(MessageHeader);
+
+	const FOutboundMessage OutboundMessage(RemoteEndpoint, MetaBuffer);
+	
+	EnqueueOutboundMessage(OutboundMessage);
+
+	UE_LOG(LogQuicMessagingTransport, Verbose,
+		TEXT("[QuicServer] Sent hello to handler %s."), *RemoteEndpoint.ToString());
+}
+
+
+void FQuicServer::SendBye()
+{
+	if (IsStopping() || ConnectionState != EQuicServerState::Running)
+	{
+		return;
+	}
+
+	TArray<FIPv4Endpoint> HandlerEndpoints;
+	Handlers.GenerateKeyArray(HandlerEndpoints);
+
+	for (FIPv4Endpoint& HandlerEndpoint : HandlerEndpoints)
+	{
+		SendBye(HandlerEndpoint);
+	}
+}
+
+
+void FQuicServer::SendBye(const FIPv4Endpoint& HandlerEndpoint)
+{
+	if (IsStopping() || ConnectionState != EQuicServerState::Running)
+	{
+		return;
+	}
+
+	FMessageHeader MessageHeader(EQuicMessageType::Bye, FGuid(), LocalNodeId, 0);
+	const FQuicPayloadPtr MetaBuffer = SerializeHeaderDelegate.Execute(MessageHeader);
+
+	const FOutboundMessage OutboundMessage(HandlerEndpoint, MetaBuffer);
+
+	HQUIC HandlerConnection;
+
+	{
+		FScopeLock HandlersLock(&HandlersCS);
+
+		const FQuicServerHandlerPtr* ServerHandler
+			= Handlers.Find(OutboundMessage.Recipient);
+
+		if (!ServerHandler)
+		{
+			UE_LOG(LogQuicMessagingTransport, Error,
+				TEXT("[QuicServer] Could not find server handler when sending bye message: %s"),
+				*OutboundMessage.Recipient.ToString());
+
+			return;
+		}
+
+		HandlerConnection = (*ServerHandler)->Connection;
+	}
+
+	SendToStream(HandlerConnection, OutboundMessage);
+
+	UE_LOG(LogQuicMessagingTransport, Verbose,
+		TEXT("[QuicServer] Sent bye to handler %s."), *HandlerEndpoint.ToString());
+}
+
+
+QUIC_STATUS FQuicServer::ServerListenerCallback(HQUIC Listener,
+	void* Context, QUIC_LISTENER_EVENT* Event)
+{
+	UNREFERENCED_PARAMETER(Listener);
+
+	FQuicServer* ServerContext = static_cast<FQuicServer*>(Context);
+
+	QUIC_STATUS Status = QUIC_STATUS_NOT_SUPPORTED;
+
+	/**
+	 * Connection attempt by a client
+	 */
+	if (Event->Type == QUIC_LISTENER_EVENT_NEW_CONNECTION)
+	{
+		Status = ServerContext->OnListenerNewConnection(Event);
+	}
+
+	return Status;
+}
+
+
+QUIC_STATUS FQuicServer::ServerConnectionCallback(HQUIC Connection,
+	void* Context, QUIC_CONNECTION_EVENT* Event)
+{
+	FQuicServer* ServerContext = static_cast<FQuicServer*>(Context);
+
+	switch (Event->Type)
+	{
+
+	case QUIC_CONNECTION_EVENT_CONNECTED:
+	
+		/**
+		 * Connection handshake has completed
+		 */
+
+		ServerContext->OnConnectionConnected(Connection);
+
+		break;
+	
+
+	case QUIC_CONNECTION_EVENT_SHUTDOWN_INITIATED_BY_TRANSPORT:
+	
+		/**
+		 * Idle timeout killed the connection
+		 */
+
+		ServerContext->OnConnectionShutdownByTransport(Connection, Event->SHUTDOWN_INITIATED_BY_TRANSPORT.Status);
+
+		break;
+	
+
+	case QUIC_CONNECTION_EVENT_SHUTDOWN_INITIATED_BY_PEER:
+	
+		/**
+		 * Peer killed the connection
+		 */
+
+		ServerContext->OnConnectionShutdownByPeer(Connection, Event->SHUTDOWN_INITIATED_BY_PEER.ErrorCode);
+
+		break;
+	
+
+	case QUIC_CONNECTION_EVENT_SHUTDOWN_COMPLETE:
+	
+		/**
+		 * Connection shutdown complete, ready to cleanup
+		 */
+
+		ServerContext->OnConnectionShutdownComplete(Connection, Event->SHUTDOWN_COMPLETE.AppCloseInProgress);
+
+		break;
+	
+
+	case QUIC_CONNECTION_EVENT_PEER_STREAM_STARTED:
+	{
+		/**
+		 * Peer started a stream, set local stream callback handler
+		 */
+
+		const QUIC_ADDR HandlerAddress = ServerContext->GetRemoteAddress(Connection);
+		FIPv4Endpoint* PeerAddress = new FIPv4Endpoint(
+			ServerContext->GetEndpointFromAddress(HandlerAddress));
+
+		FStreamCallbackInfo* StreamCallbackInfo = new FStreamCallbackInfo();
+		StreamCallbackInfo->Endpoint = ServerContext;
+		StreamCallbackInfo->PeerAddress = PeerAddress;
+
+		ServerContext->OnConnectionPeerStreamStarted(Connection, Event->PEER_STREAM_STARTED.Stream);
+
+		ServerContext->MsQuic->SetCallbackHandler(Event->PEER_STREAM_STARTED.Stream,
+			reinterpret_cast<void*>(ServerStreamCallback), StreamCallbackInfo);
+
+		break;
+	}
+
+	case QUIC_CONNECTION_EVENT_RESUMED:
+
+		/**
+		 * TLS resumption of previous session succeeded
+		 */
+
+		ServerContext->OnConnectionResumed();
+
+	case QUIC_CONNECTION_EVENT_DATAGRAM_RECEIVED:
+	case QUIC_CONNECTION_EVENT_DATAGRAM_SEND_STATE_CHANGED:
+	case QUIC_CONNECTION_EVENT_DATAGRAM_STATE_CHANGED:
+	case QUIC_CONNECTION_EVENT_IDEAL_PROCESSOR_CHANGED:
+	case QUIC_CONNECTION_EVENT_LOCAL_ADDRESS_CHANGED:
+	case QUIC_CONNECTION_EVENT_PEER_ADDRESS_CHANGED:
+	case QUIC_CONNECTION_EVENT_PEER_CERTIFICATE_RECEIVED:
+	case QUIC_CONNECTION_EVENT_PEER_NEEDS_STREAMS:
+	case QUIC_CONNECTION_EVENT_RESUMPTION_TICKET_RECEIVED:
+	case QUIC_CONNECTION_EVENT_STREAMS_AVAILABLE:
+
+		break;
+
+	default:
+
+		UE_LOG(LogQuicMessagingTransport, Warning,
+			TEXT("[QuicServer] Received unhandled connection event: %d"), Event->Type);
+
+		break;
+
+	}
+
+	return QUIC_STATUS_SUCCESS;
+}
+
+
+QUIC_STATUS FQuicServer::ServerStreamCallback(HQUIC Stream,
+	void* Context, QUIC_STREAM_EVENT* Event)
+{
+	const FStreamCallbackInfo* CallbackInfo = static_cast<FStreamCallbackInfo*>(Context);
+	FQuicServer* ServerContext = static_cast<FQuicServer*>(CallbackInfo->Endpoint);
+
+	switch (Event->Type)
+	{
+
+	case QUIC_STREAM_EVENT_SEND_COMPLETE:
+
+		/**
+		 * Previous StreamSend has completed, context returned
+		 */
+
+		ServerContext->OnStreamSendComplete(Stream, Event);
+
+		break;
+
+	case QUIC_STREAM_EVENT_RECEIVE:
+	
+		/**
+		 * Data received from the peer
+		 */
+
+		ServerContext->OnStreamReceive(Stream, *CallbackInfo->PeerAddress, Event);
+
+		return QUIC_STATUS_PENDING;
+
+	case QUIC_STREAM_EVENT_PEER_SEND_ABORTED:
+
+		/**
+		 * Peer aborted its stream send direction
+		 */
+
+		ServerContext->OnStreamPeerSendAborted(Stream);
+
+		break;
+
+	case QUIC_STREAM_EVENT_PEER_SEND_SHUTDOWN:
+
+		/**
+		 * Peer gracefully shut down its stream send direction
+		 */
+
+		break;
+
+	case QUIC_STREAM_EVENT_SHUTDOWN_COMPLETE:
+	
+		/**
+		 * Both stream directions have been shut down, safe to cleanup
+		 */
+
+		ServerContext->OnStreamShutdownComplete(Stream);
+		delete CallbackInfo;
+
+		break;
+
+	case QUIC_STREAM_EVENT_START_COMPLETE:
+	case QUIC_STREAM_EVENT_PEER_ACCEPTED:
+	case QUIC_STREAM_EVENT_PEER_RECEIVE_ABORTED:
+	case QUIC_STREAM_EVENT_IDEAL_SEND_BUFFER_SIZE:
+	case QUIC_STREAM_EVENT_SEND_SHUTDOWN_COMPLETE:
+
+		break;
+
+	default:
+
+		UE_LOG(LogQuicMessagingTransport, Warning,
+			TEXT("[QuicServer] Received unhandled stream event: %d"), Event->Type);
+
+		break;
+
+	}
+
+	return QUIC_STATUS_SUCCESS;
+}
+
+
+QUIC_STATUS FQuicServer::OnListenerNewConnection(QUIC_LISTENER_EVENT* Event)
+{
+	const HQUIC HandlerConnection = Event->NEW_CONNECTION.Connection;
+	const QUIC_ADDR HandlerAddress = GetRemoteAddress(Event->NEW_CONNECTION.Connection);
+	const FIPv4Endpoint HandlerEndpoint = GetEndpointFromAddress(HandlerAddress);
+
+	if (!CanHandlerConnect(HandlerEndpoint))
+	{
+		return QUIC_STATUS_ABORTED;
+	}
+
+	AddConnectionAttempt(HandlerEndpoint);
+
+	if (ServerConfig->EncryptionMode == EEncryptionMode::Disabled)
+	{
+		constexpr BOOLEAN Value = true;
+
+		MsQuic->SetParam(
+			Event->NEW_CONNECTION.Connection,
+			QUIC_PARAM_CONN_DISABLE_1RTT_ENCRYPTION,
+			sizeof(Value), &Value);
+	}
+
+	MsQuic->SetCallbackHandler(Event->NEW_CONNECTION.Connection,
+		reinterpret_cast<void*>(ServerConnectionCallback), this);
+
+	const QUIC_STATUS Status = MsQuic->ConnectionSetConfiguration(
+		Event->NEW_CONNECTION.Connection, QuicConfig);
+
+	{
+		FScopeLock HandlersLock(&HandlersCS);
+
+		const TSharedPtr<FQuicServerHandler, ESPMode::ThreadSafe> Handler
+			= MakeShared<FQuicServerHandler, ESPMode::ThreadSafe>(
+				HandlerEndpoint, HandlerAddress, HandlerConnection);
+
+		Handlers.Add(HandlerEndpoint, Handler);
+	}
+
+	UE_LOG(LogQuicMessagingTransport, Verbose,
+		TEXT("[QuicServer] Connection attempt by client: %s."),
+		*HandlerEndpoint.ToString());
+
+	return Status;
+}
+
+
+void FQuicServer::OnConnectionConnected(HQUIC Connection)
+{
+	const QUIC_ADDR Handler = GetRemoteAddress(Connection);
+	const FIPv4Endpoint HandlerEndpoint = GetEndpointFromAddress(Handler);
+
+	{
+		FScopeLock HandlersLock(&HandlersCS);
+
+		if (const FQuicServerHandlerPtr* ServerHandler = Handlers.Find(HandlerEndpoint))
+		{
+			(*ServerHandler)->ConnectionState = EQuicHandlerState::Running;
+			(*ServerHandler)->bSentHello = true;
+		}
+	}
+
+	SendHello(HandlerEndpoint);
+
+	EndpointConnectedDelegate.ExecuteIfBound(HandlerEndpoint);
+
+	UE_LOG(LogQuicMessagingTransport, Verbose,
+		TEXT("[QuicServer] Handshake completed. Handler: %s."), *HandlerEndpoint.ToString());
+}
+
+
+void FQuicServer::OnConnectionShutdownByTransport(HQUIC Connection, QuicUtils::HRESULT Status)
+{
+	const QUIC_ADDR Handler = GetRemoteAddress(Connection);
+	const FIPv4Endpoint HandlerEndpoint = GetEndpointFromAddress(Handler);
+	const EQuicEndpointError HandlerError = QuicUtils::ConvertQuicStatus(Status);
+
+	{
+		FScopeLock HandlersLock(&HandlersCS);
+
+		if (const FQuicServerHandlerPtr* ServerHandler = Handlers.Find(HandlerEndpoint))
+		{
+			(*ServerHandler)->ConnectionState = EQuicHandlerState::Disconnected;
+			(*ServerHandler)->Error = HandlerError;
+		}
+	}
+
+	UE_LOG(LogQuicMessagingTransport, Verbose,
+		TEXT("[QuicServer] Connection end (by transport). Handler: %s."),
+		*HandlerEndpoint.ToString());
+
+	if (HandlerError != EQuicEndpointError::Normal)
+	{
+		UE_LOG(LogQuicMessagingTransport, Verbose,
+			TEXT("[QuicServer] Handler: %s / shutdown reason: %s."),
+			*HandlerEndpoint.ToString(), *QuicUtils::GetEndpointErrorString(HandlerError));
+	}
+}
+
+
+void FQuicServer::OnConnectionShutdownByPeer(HQUIC Connection, QUIC_UINT62 ErrorCode)
+{
+	const QUIC_ADDR Handler = GetRemoteAddress(Connection);
+	const FIPv4Endpoint HandlerEndpoint = GetEndpointFromAddress(Handler);
+
+	{
+		FScopeLock HandlersLock(&HandlersCS);
+
+		if (const FQuicServerHandlerPtr* ServerHandler = Handlers.Find(HandlerEndpoint))
+		{
+			(*ServerHandler)->ConnectionState = EQuicHandlerState::Disconnected;
+		}
+	}
+
+	UE_LOG(LogQuicMessagingTransport, Verbose,
+		TEXT("[QuicServer] Connection end (by peer). Handler: %s / Error code: %d."),
+		*HandlerEndpoint.ToString(), ErrorCode);
+}
+
+
+void FQuicServer::OnConnectionShutdownComplete(HQUIC Connection, BOOLEAN AppCloseInProgress)
+{
+	FScopeLock HandlersLock(&HandlersCS);
+
+	const QUIC_ADDR Handler = GetRemoteAddress(Connection);
+	const FIPv4Endpoint HandlerEndpoint = GetEndpointFromAddress(Handler);
+
+	if (!AppCloseInProgress)
+	{
+		Handlers.Remove(GetEndpointFromAddress(Handler));
+
+		HandlersLock.Unlock();
+
+		EndpointDisconnectedDelegate.ExecuteIfBound(
+			HandlerEndpoint, EQuicEndpointError::Normal);
+
+		if (MsQuic)
+		{
+			MsQuic->ConnectionClose(Connection);
+		}
+	}
+
+	UE_LOG(LogQuicMessagingTransport, Verbose,
+		TEXT("[QuicServer] Connection shutdown complete. Handler: %s."),
+		*HandlerEndpoint.ToString());
+}
+
+
+void FQuicServer::OnConnectionPeerStreamStarted(HQUIC Connection, HQUIC Stream)
+{
+	const uint64 StreamId = GetStreamId(Stream);
+	const QUIC_ADDR HandlerAddress = GetRemoteAddress(Connection);
+	const FIPv4Endpoint HandlerEndpoint = GetEndpointFromAddress(HandlerAddress);
+
+	{
+		const FQuicPayloadPtr InboundData = MakeShared<TArray<uint8>, ESPMode::ThreadSafe>();
+
+		FScopeLock InboundLock(&InboundBuffersCS);
+
+		const FInboundBuffer InboundBuffer(HandlerAddress, HandlerEndpoint, InboundData);
+		const FAddrStream AddressStreamKey(HandlerEndpoint, StreamId);
+
+		InboundBuffers.Add(AddressStreamKey, InboundBuffer);
+	}
+
+	UE_LOG(LogQuicMessagingTransport, VeryVerbose,
+		TEXT("[QuicServer] Peer stream started. Handler: %s / StreamId: %d."),
+		*HandlerEndpoint.ToString(), StreamId);
+}
+
+
+void FQuicServer::OnConnectionResumed()
+{
+	UE_LOG(LogQuicMessagingTransport, Verbose,
+		TEXT("[QuicServer] Connection resumed."));
+}
+
+
+void FQuicServer::OnStreamShutdownComplete(HQUIC Stream) const
+{
+	if (MsQuic)
+	{
+		MsQuic->StreamClose(Stream);
+	}
+}
+
+
+void FQuicServer::SendToStream(
+	const HQUIC Connection, const FOutboundMessage OutboundMessage)
+{
+	HQUIC Stream = nullptr;
+	QUIC_STATUS Status;
+
+	FIPv4Endpoint* PeerAddress = new FIPv4Endpoint(OutboundMessage.Recipient);
+
+	FStreamCallbackInfo* StreamCallbackInfo = new FStreamCallbackInfo();
+	StreamCallbackInfo->Endpoint = this;
+	StreamCallbackInfo->PeerAddress = PeerAddress;
+
+	/**
+	 * Create/allocate a new unidirectional stream
+	 */
+	if (QUIC_FAILED(Status = MsQuic->StreamOpen(Connection, QUIC_STREAM_OPEN_FLAG_UNIDIRECTIONAL,
+		ServerStreamCallback, StreamCallbackInfo, &Stream)))
+	{
+		UE_LOG(LogQuicMessagingTransport, Error,
+			TEXT("[QuicServer] Could not open stream: %d."), QuicUtils::ConvertQuicStatus(Status));
+		return;
+	}
+
+	/**
+	 * Start the stream
+	 */
+	if (QUIC_FAILED(Status = MsQuic->StreamStart(Stream, QUIC_STREAM_START_FLAG_NONE)))
+	{
+		UE_LOG(LogQuicMessagingTransport, Error,
+			TEXT("[QuicServer] Could not start stream: %d."), QuicUtils::ConvertQuicStatus(Status));
+
+		MsQuic->StreamClose(Stream);
+		return;
+	}
+
+	const uint64 StreamId = GetStreamId(Stream);
+
+	/**
+	 * Send buffers over the stream
+	 */
+	const FQuicBufferRef MetaBuffer = MakeShared<QUIC_BUFFER, ESPMode::ThreadSafe>();
+	const FQuicBufferRef MessageBuffer = MakeShared<QUIC_BUFFER, ESPMode::ThreadSafe>();
+
+	MetaBuffer->Buffer = OutboundMessage.SerializedHeader->GetData();
+	MetaBuffer->Length = OutboundMessage.SerializedHeader->Num();
+
+	RegisterStreamMessage(StreamId, *PeerAddress,
+		!OutboundMessage.bHasPayload, OutboundMessage, MetaBuffer, MessageBuffer);
+
+	if (QUIC_FAILED(Status = MsQuic->StreamSend(Stream, &*MetaBuffer, 1,
+		OutboundMessage.bHasPayload ? QUIC_SEND_FLAG_NONE : QUIC_SEND_FLAG_FIN, PeerAddress)))
+	{
+		UE_LOG(LogQuicMessagingTransport, Error,
+			TEXT("[QuicServer] Could not send meta on stream: %d."),
+			QuicUtils::ConvertQuicStatus(Status));
+		
+		MsQuic->StreamShutdown(Stream, QUIC_STREAM_SHUTDOWN_FLAG_ABORT, 0);
+		return;
+	}
+
+	// This message is header-only
+	if (!OutboundMessage.bHasPayload)
+	{
+		return;
+	}
+
+	MessageBuffer->Buffer = OutboundMessage.SerializedPayload->GetData();
+	MessageBuffer->Length = OutboundMessage.SerializedPayload->Num();
+
+	if (QUIC_FAILED(Status = MsQuic->StreamSend(Stream,
+		&*MessageBuffer, 1, QUIC_SEND_FLAG_FIN, PeerAddress)))
+	{
+		UE_LOG(LogQuicMessagingTransport, Error,
+			TEXT("[QuicServer] Could not send message on stream: %d."),
+			QuicUtils::ConvertQuicStatus(Status));
+		
+		MsQuic->StreamShutdown(Stream, QUIC_STREAM_SHUTDOWN_FLAG_ABORT, 0);
+	}
+}
+
+
+bool FQuicServer::ReceiveFromStream(const FInboundQuicBuffer& QuicBuffer)
+{
+	FScopeLock InboundLock(&InboundBuffersCS);
+
+	const FAddrStream AddressStreamKey(
+		QuicBuffer.PeerAddress, QuicBuffer.StreamId);
+
+	FInboundBuffer* Buffer = InboundBuffers.Find(AddressStreamKey);
+
+	if (!Buffer)
+	{
+		UE_LOG(LogQuicMessagingTransport, Error,
+			TEXT("[QuicServer] Invalid inbound buffer in ReceiveFromStream %d."),
+			QuicBuffer.StreamId);
+
+		return false;
+	}
+
+	const FQuicPayloadPtr& InboundData = Buffer->InboundData;
+
+	for (const TPair<uint8*, uint32> InboundBuffer : QuicBuffer.Data)
+	{
+		Buffer->InboundData->Append(InboundBuffer.Key, InboundBuffer.Value);
+	}
+
+	// Retrieve message header
+	if (!Buffer->bHeaderDeserialized
+		&& InboundData->Num() >= QUIC_MESSAGE_HEADER_SIZE)
+	{
+		const FQuicPayloadPtr HeaderData = MakeShared<TArray<uint8>, ESPMode::ThreadSafe>();
+		HeaderData->Append(InboundData->GetData(), QUIC_MESSAGE_HEADER_SIZE);
+		const FMessageHeader MessageHeader = DeserializeHeaderDelegate.Execute(HeaderData);
+
+		Buffer->MessageHeader = MessageHeader;
+		Buffer->bHeaderDeserialized = true;
+
+		InboundData->RemoveAt(0, QUIC_MESSAGE_HEADER_SIZE);
+
+		const uint32& SerializedMessageSize = MessageHeader.SerializedMessageSize;
+		InboundData->Reserve(SerializedMessageSize);
+
+		// Handler is not authenticated and trying to send message with type Data
+		if (AuthenticationMode == EAuthenticationMode::Enabled
+			&& MessageHeader.MessageType == EQuicMessageType::Data)
+		{
+			{
+				const FIPv4Endpoint HandlerAddress
+					= GetEndpointFromAddress(Buffer->HandlerAddress);
+
+				FScopeLock HandlersLock(&HandlersCS);
+
+				const FQuicServerHandlerPtr* Handler = Handlers.Find(HandlerAddress);
+
+				if (!Handler)
+				{
+					UE_LOG(LogQuicMessagingTransport, Error,
+						TEXT("[QuicServer] Cannot check authentication status of invalid handler %s."),
+						*HandlerAddress.ToString());
+
+					return false;
+				}
+
+				if (!(*Handler)->bAuthenticated)
+				{
+					UE_LOG(LogQuicMessagingTransport, Error,
+						TEXT("[QuicServer] Unauthenticated handler sending data."));
+
+					return false;
+				}
+			}
+		}
+	}
+
+	// Check authentication message size
+	if (Buffer->bHeaderDeserialized)
+	{
+		if (Buffer->MessageHeader.MessageType == EQuicMessageType::Authentication
+			&& InboundData->Num() > static_cast<int32>(ServerConfig->MaxAuthenticationMessageSize))
+		{
+			UE_LOG(LogQuicMessagingTransport, Error,
+				TEXT("[QuicServer] Authentication message size too large."));
+
+			return false;
+		}
+	}
+
+	InboundLock.Unlock();
+
+	if (MsQuic)
+	{
+		MsQuic->StreamReceiveComplete(QuicBuffer.Stream, QuicBuffer.BufferLength);
+	}
+
+	return true;
+}
+
+
+void FQuicServer::ReachedEndOfStream(const FInboundQuicBuffer& QuicBuffer)
+{
+	FInboundBuffer InboundBuffer;
+
+	{
+		const FAddrStream AddressStreamKey(QuicBuffer.PeerAddress, QuicBuffer.StreamId);
+
+		FScopeLock InboundLock(&InboundBuffersCS);
+
+		if (!InboundBuffers.RemoveAndCopyValue(AddressStreamKey, InboundBuffer))
+		{
+			UE_LOG(LogQuicMessagingTransport, Error,
+				TEXT("[QuicServer] Invalid inbound buffer in ReachedEndOfStream."));
+
+			return;
+		}
+	}
+
+	const FQuicPayloadPtr& InboundData = InboundBuffer.InboundData;
+	const FMessageHeader& MessageHeader = InboundBuffer.MessageHeader;
+
+	// Transmitted message size smaller than expected
+	if (InboundData->Num() < static_cast<int32>(MessageHeader.SerializedMessageSize))
+	{
+		UE_LOG(LogQuicMessagingTransport, Error,
+			TEXT("[QuicServer] Transmitted message size smaller than expected %d / %d / %d."),
+			QuicBuffer.StreamId, InboundData->Num(), MessageHeader.SerializedMessageSize);
+
+		return;
+	}
+
+	FInboundMessage InboundMessage(InboundData, MessageHeader,
+		InboundBuffer.HandlerEndpoint, Endpoint);
+
+	InboundMessages.Enqueue(InboundMessage);
+}
+
+
+QUIC_ADDR FQuicServer::GetRemoteAddress(HQUIC Connection) const
+{
+	QUIC_ADDR RemoteAddress;
+	uint32_t RemoteAddressSize = sizeof(RemoteAddress);
+
+	if (MsQuic)
+	{
+		QUIC_STATUS Status;
+
+		if (QUIC_FAILED(Status = MsQuic->GetParam(Connection,
+			QUIC_PARAM_CONN_REMOTE_ADDRESS, &RemoteAddressSize, &RemoteAddress)))
+		{
+			UE_LOG(LogQuicMessagingTransport, Error,
+				TEXT("[QuicServer] Could not get QUIC address from connection: %s."),
+				*QuicUtils::ConvertResult(Status));
+		}
+	}
+
+	return RemoteAddress;
+}
+
+
+FIPv4Endpoint FQuicServer::ClearEndpointPort(const FIPv4Endpoint& Endpoint)
+{
+	FIPv4Endpoint ClearedEndpoint = Endpoint;
+	ClearedEndpoint.Port = 0;
+
+	return ClearedEndpoint;
+}
+
diff --git a/Engine/Plugins/Experimental/QuicMessaging/Source/QuicMessagingTransport/Private/QuicServer.h b/Engine/Plugins/Experimental/QuicMessaging/Source/QuicMessagingTransport/Private/QuicServer.h
new file mode 100644
--- /dev/null
+++ b/Engine/Plugins/Experimental/QuicMessaging/Source/QuicMessagingTransport/Private/QuicServer.h
@@ -0,0 +1,391 @@
+// Copyright Epic Games, Inc. All Rights Reserved.
+
+#pragma once
+
+#include <functional>
+
+#include "Containers/Ticker.h"
+
+#include "QuicServerHandler.h"
+
+#include "QuicEndpoint.h"
+#include "QuicQueues.h"
+#include "QuicMessages.h"
+
+#include "QuicUtils.h"
+
+
+// -Woverloaded-virtual:
+// Clang thinks the SendHello and SendBye methods are hidden
+// overloaded virtual functions because the ones in QuicServer
+// require one argument whereas the ones in QuicEndpoint do not
+#ifdef __clang__
+#pragma clang diagnostic push
+#pragma clang diagnostic ignored "-Woverloaded-virtual"
+#endif
+
+
+typedef TSharedPtr<FQuicServerHandler, ESPMode::ThreadSafe> FQuicServerHandlerPtr;
+
+
+/**
+ * Implements the QuicServer.
+ */
+class FQuicServer
+	: public FQuicEndpoint
+	, public FQuicQueues
+{
+
+public:
+
+	FQuicServer() = delete;
+
+	/**
+	 * Creates a new QUIC server.
+	 */
+	FQuicServer(TSharedRef<FQuicServerConfig> InServerConfig,
+		HQUIC InQuicConfig, const QUIC_API_TABLE* InMsQuic, HQUIC InRegistration,
+		const TCHAR* InThreadName, const QUIC_BUFFER InAlpn)
+		: FQuicEndpoint(InServerConfig, InQuicConfig,
+			InMsQuic, InRegistration, InThreadName)
+		, ServerListener(nullptr)
+		, ServerConnection(nullptr)
+		, Alpn(InAlpn)
+		, AuthenticationMode(InServerConfig->AuthenticationMode)
+		, ServerConfig(InServerConfig)
+		, Handlers(TMap<FIPv4Endpoint, FQuicServerHandlerPtr>())
+		, ConnectionCooldown(TMap<FIPv4Endpoint, uint64>())
+		, ConnectionAttempts(TMap<FIPv4Endpoint, FConnectionAttempts>())
+	{ }
+
+protected:
+
+	/**
+	 * Starts the configured QUIC server.
+	 */
+	virtual void StartEndpoint() override;
+
+public:
+
+	/**
+	 * Stops the QUIC server.
+	 */
+	virtual void StopEndpoint() override;
+
+	/**
+	 * Indicates if the server is ready to send/receive messages.
+	 */
+	virtual bool IsEndpointReady() override;
+
+public:
+
+	/**
+	 * Periodically called to update the endpoint (consume messages).
+	 */
+	virtual void UpdateEndpoint() override;
+
+protected:
+
+	/**
+	 * Dequeue inbound messages until MaxConsumeTime is reached.
+	 * Inbound messages are passed to QuicEndpointManager for validation.
+	 */
+	virtual void ConsumeInboundMessages() override;
+
+	/**
+	 * Dequeue inbound buffers until MaxConsumeTime is reached.
+	 * Inbound buffers via QUIC are copied into our buffers until the stream end is indicated.
+	 */
+	virtual void ProcessInboundBuffers() override;
+
+	/**
+	 * Dequeue outbound messages until MaxConsumeTime is reached.
+	 * Outbound messages will then be sent via QUIC.
+	 */
+	virtual void ConsumeOutboundMessages() override;
+
+protected:
+
+	/**
+	 * Determines if a handler can connect by checking the connection cooldown list.
+	 *
+	 * @param ConnectingEndpoint The connecting endpoint's address
+	 */
+	bool CanHandlerConnect(const FIPv4Endpoint& ConnectingEndpoint) const;
+
+	/**
+	 * Adds a new connection attempt and triggers connection cooldown if enabled.
+	 *
+	 * @param ConnectingEndpoint The connecting endpoint's address
+	 */
+	void AddConnectionAttempt(const FIPv4Endpoint& ConnectingEndpoint);
+
+	/**
+	 * Regularly checks and cleans up connection cooldown entries.
+	 */
+	void CheckConnectionCooldown();
+
+	/**
+	 * Collects, converts and passes QUIC connection stats to the EndpointManager.
+	 */
+	virtual void CollectStatistics() override;
+
+public:
+
+	/**
+	 * Marks a handler as authenticated.
+	 * 
+	 * @param HandlerEndpoint The handler endpoint
+	 * @see QuicEndpointManager / Authentication
+	 */
+	virtual void SetHandlerAuthenticated(const FIPv4Endpoint& HandlerEndpoint) override;
+
+	/**
+	 * Disconnects a handler.
+	 *
+	 * @param HandlerEndpoint The handler endpoint
+	 */
+	virtual void DisconnectHandler(const FIPv4Endpoint& HandlerEndpoint) override;
+
+protected:
+
+	/**
+	 * Shutdown every handler's connection.
+	 */
+	void ShutdownAllHandlers() const;
+
+	/**
+	 * Shutdown a handler's connection.
+	 *
+	 * @param HandlerConnection The handler connection
+	 */
+	void ShutdownHandler(const HQUIC HandlerConnection) const;
+
+public:
+
+	/**
+	 * Enqueues an outbound message on the server.
+	 *
+	 * @param OutboundMessage The outbound message
+	 */
+	virtual void EnqueueOutboundMessage(const FOutboundMessage& OutboundMessage) override;
+
+	/**
+	 * Sends a hello message to the remote endpoint.
+	 *
+	 * @param RemoteEndpoint The remote endpoint
+	 */
+	virtual void SendHello(FIPv4Endpoint RemoteEndpoint) override;
+
+	/**
+	 * Sends a bye message to all handlers.
+	 */
+	virtual void SendBye() override;
+
+	/**
+	 * Sends a bye message to the specified handler.
+	 *
+	 * @param HandlerEndpoint The handler endpoint
+	 */
+	void SendBye(const FIPv4Endpoint& HandlerEndpoint);
+
+protected:
+
+	/**
+	 * Callback for server listeners.
+	 *
+	 * @param Listener The QUIC listener handle
+	 * @param Context The app context (FQuicServer)
+	 * @param Event The listener event
+	 */
+	static QUIC_STATUS ServerListenerCallback(HQUIC Listener,
+		void* Context, QUIC_LISTENER_EVENT* Event);
+
+	/**
+	 * Callback for server connections.
+	 *
+	 * @param Connection The QUIC connection handle
+	 * @param Context The app context (FQuicServer)
+	 * @param Event The connection event
+	 */
+	static QUIC_STATUS ServerConnectionCallback(HQUIC Connection,
+		void* Context, QUIC_CONNECTION_EVENT* Event);
+
+	/**
+	 * Callback for server streams.
+	 *
+	 * @param Stream The QUIC stream handle
+	 * @param Context The app context (FQuicServer)
+	 * @param Event The stream event
+	 */
+	static QUIC_STATUS ServerStreamCallback(HQUIC Stream,
+		void* Context, QUIC_STREAM_EVENT* Event);
+
+protected:
+
+	/**
+	 * Handles a connection attempt by a client.
+	 *
+	 * @param Event The listener event
+	 */
+	QUIC_STATUS OnListenerNewConnection(QUIC_LISTENER_EVENT* Event);
+
+	/**
+	 * Handles a connection after the handshake was completed.
+	 *
+	 * @param Connection The connection handle
+	 */
+	void OnConnectionConnected(HQUIC Connection);
+
+	/**
+	 * Handles a connection after it was killed due to idle timeout.
+	 *
+	 * @param Connection The connection handle
+	 * @param Status The shutdown event status
+	 */
+	void OnConnectionShutdownByTransport(HQUIC Connection, QuicUtils::HRESULT Status);
+
+	/**
+	 * Handles a connection after it was killed by the peer.
+	 *
+	 * @param Connection The connection handle
+	 * @param ErrorCode The QUIC error code
+	 */
+	void OnConnectionShutdownByPeer(HQUIC Connection, QUIC_UINT62 ErrorCode);
+
+	/**
+	 * Handles a connection after its shutdown was completed, ready for cleanup.
+	 *
+	 * @param Connection The connection handle
+	 * @param AppCloseInProgress Boolean indicating whether app close is already in progress
+	 */
+	void OnConnectionShutdownComplete(HQUIC Connection, BOOLEAN AppCloseInProgress);
+
+	/**
+	 * Handles when a peer started a stream.
+	 *
+	 * @param Connection The connection handle
+	 * @param Stream The stream handle
+	 */
+	void OnConnectionPeerStreamStarted(HQUIC Connection, HQUIC Stream);
+
+	/**
+	 * Handles a succeeded TLS resumption of a previous session.
+	 */
+	void OnConnectionResumed();
+
+	/**
+	 * Closes a stream when both directions have been shut down.
+	 *
+	 * @param Stream The stream handle
+	 */
+	void OnStreamShutdownComplete(HQUIC Stream) const;
+
+protected:
+
+	/**
+	 * Sends the outbound buffers to the remote endpoint.
+	 *
+	 * @param Connection The QUIC connection handle
+	 * @param OutboundMessage The outbound message
+	 */
+	void SendToStream(
+		const HQUIC Connection, const FOutboundMessage OutboundMessage);
+
+protected:
+
+	/**
+	 * Copy data from a stream buffer.
+	 *
+	 * @param QuicBuffer The inbound QUIC buffer
+	 *
+	 * @return Whether the operation succeeded
+	 */
+	bool ReceiveFromStream(const FInboundQuicBuffer& QuicBuffer);
+
+	/**
+	 * Prepares inbound data for processing.
+	 *
+	 * @param QuicBuffer The inbound QUIC buffer
+	 */
+	void ReachedEndOfStream(const FInboundQuicBuffer& QuicBuffer);
+
+protected:
+
+	/**
+	 * Get the remote address from a connection.
+	 *
+	 * @param Connection The QUIC connection handle
+	 */
+	QUIC_ADDR GetRemoteAddress(HQUIC Connection) const;
+
+	/**
+	 * Sets the endpoints port to zero.
+	 *
+	 * @param Endpoint The endpoint address
+	 */
+	static FIPv4Endpoint ClearEndpointPort(const FIPv4Endpoint& Endpoint);
+
+protected:
+
+	/** Holds the QUIC API server listener handle. */
+	HQUIC ServerListener;
+
+	/** Holds the QUIC API server connection handle. */
+	HQUIC ServerConnection;
+
+	/** Current server state. */
+	EQuicServerState ConnectionState = EQuicServerState::Starting;
+
+	/** Holds the protocol name for ALPN. */
+	QUIC_BUFFER Alpn;
+
+	/** Holds the server authentication mode. */
+	EAuthenticationMode AuthenticationMode;
+
+	/** Holds the server config. */
+	TSharedRef<FQuicServerConfig> ServerConfig;
+
+protected:
+
+	/** Holds the QuicServerHandlers. */
+	TMap<FIPv4Endpoint, FQuicServerHandlerPtr> Handlers;
+
+	/** Holds the handlers critical section. */
+	mutable FCriticalSection HandlersCS;
+
+	/** Holds the handle for the connection cooldown check. */
+	FTSTicker::FDelegateHandle ConnectionCooldownCheckHandle;
+
+	/** Holds the connection cooldown critical section. */
+	mutable FCriticalSection ConnectionCooldownCS;
+
+	/** Holds a map of IPv4 addresses and connection cooldown timestamps. */
+	TMap<FIPv4Endpoint, uint64> ConnectionCooldown;
+
+	/** Struct for data concerning connection attempts. */
+	struct FConnectionAttempts
+	{
+		/** Holds the number of previous cooldowns. */
+		uint32 NumCooldowns;
+
+		/** Holds the attempt timestamps. */
+		TArray<uint64> AttemptTimestamps;
+
+		/** Default constructor. */
+		FConnectionAttempts()
+			: NumCooldowns(0)
+			, AttemptTimestamps(TArray<uint64>())
+		{ }
+	};
+
+	/** Holds a map of connection attempts. */
+	TMap<FIPv4Endpoint, FConnectionAttempts> ConnectionAttempts;
+
+};
+
+
+#ifdef __clang__
+#pragma clang diagnostic pop
+#endif
+
diff --git a/Engine/Plugins/Experimental/QuicMessaging/Source/QuicMessagingTransport/Private/QuicServerHandler.h b/Engine/Plugins/Experimental/QuicMessaging/Source/QuicMessagingTransport/Private/QuicServerHandler.h
new file mode 100644
--- /dev/null
+++ b/Engine/Plugins/Experimental/QuicMessaging/Source/QuicMessagingTransport/Private/QuicServerHandler.h
@@ -0,0 +1,109 @@
+// Copyright Epic Games, Inc. All Rights Reserved.
+
+#pragma once
+
+#include "QuicIncludes.h"
+#include "QuicFlags.h"
+
+#include "Interfaces/IPv4/IPv4Endpoint.h"
+#include "HAL/PlatformTime.h"
+
+
+#define QUIC_STATS_RTTAVG_RESET_AFTER 300000000
+
+
+struct FQuicServerHandler
+{
+
+	/** Holds a flag indicating whether a hello message was sent. */
+	bool bSentHello;
+
+	/** Holds a flag indicating whether the handler is authenticated. */
+	bool bAuthenticated;
+
+	/** Holds the current handler state. */
+	EQuicHandlerState ConnectionState;
+
+	/** Holds the handler endpoint. */
+	FIPv4Endpoint Endpoint;
+
+	/** Holds the handler address. */
+	QUIC_ADDR Address;
+
+	/** Holds the handler connection. */
+	HQUIC Connection;
+
+	/** Holds the handler error. */
+	EQuicEndpointError Error;
+
+	/** Holds the latest RTT measurements. */
+	TArray<uint32> RttMeasurements;
+
+	/** Holds the time when the RTT measurement sequence started. */
+	double RttStart;
+
+	/** Default constructor. */
+	FQuicServerHandler()
+		: bSentHello(false)
+		, bAuthenticated(false)
+		, ConnectionState(EQuicHandlerState::Starting)
+		, Endpoint(FIPv4Endpoint::Any)
+		, Address({})
+		, Connection(nullptr)
+		, Error(EQuicEndpointError::Normal)
+		, RttMeasurements(TArray<uint32>())
+		, RttStart(FPlatformTime::Seconds())
+	{ }
+
+	/** Creates and initializes a new instance. */
+	FQuicServerHandler(const FIPv4Endpoint& InEndpoint,
+		QUIC_ADDR InAddress, const HQUIC InConnection)
+		: bSentHello(false)
+		, bAuthenticated(false)
+		, ConnectionState(EQuicHandlerState::Starting)
+		, Endpoint(InEndpoint)
+		, Address(InAddress)
+		, Connection(InConnection)
+		, Error(EQuicEndpointError::Normal)
+		, RttMeasurements(TArray<uint32>())
+		, RttStart(FPlatformTime::Seconds())
+	{ }
+
+	bool IsStarting() const
+	{
+		return ConnectionState == EQuicHandlerState::Starting;
+	}
+
+	bool IsRunning() const
+	{
+		return ConnectionState == EQuicHandlerState::Running;
+	}
+
+	bool IsStopping() const
+	{
+		return ConnectionState == EQuicHandlerState::Stopping;
+	}
+
+	bool NeedRttReset() const
+	{
+		return (FPlatformTime::Seconds() - RttStart) > QUIC_STATS_RTTAVG_RESET_AFTER;
+	}
+
+	/**
+	 * Gets the average RTT measurement.
+	 */
+	double GetAverageRtt() const
+	{
+		uint32 RttSum = 0;
+
+		for (uint32 RttPoint : RttMeasurements)
+		{
+			RttSum += RttPoint;
+		}
+
+		return RttSum / RttMeasurements.Num();
+	}
+
+};
+
+
diff --git a/Engine/Plugins/Experimental/QuicMessaging/Source/QuicMessagingTransport/Private/QuicUtils.cpp b/Engine/Plugins/Experimental/QuicMessaging/Source/QuicMessagingTransport/Private/QuicUtils.cpp
new file mode 100644
--- /dev/null
+++ b/Engine/Plugins/Experimental/QuicMessaging/Source/QuicMessagingTransport/Private/QuicUtils.cpp
@@ -0,0 +1,261 @@
+// Copyright Epic Games, Inc. All Rights Reserved.
+
+#include "QuicUtils.h"
+#include "QuicIncludes.h"
+
+
+namespace QuicUtils
+{
+
+	FString ConvertResult(HRESULT Result)
+	{
+		return FString::Printf(TEXT("0x%lx"), Result);
+	}
+
+
+	EQuicEndpointError ConvertQuicStatus(HRESULT Status)
+	{
+		EQuicEndpointError EndpointError;
+
+		switch (Status)
+		{
+
+		case QUIC_STATUS_SUCCESS:
+
+			EndpointError = EQuicEndpointError::Normal;
+			break;
+
+		case QUIC_STATUS_OUT_OF_MEMORY:
+
+			EndpointError = EQuicEndpointError::OutOfMemory;
+			break;
+
+		case QUIC_STATUS_INVALID_PARAMETER:
+
+			EndpointError = EQuicEndpointError::InvalidParameter;
+			break;
+
+		case QUIC_STATUS_INVALID_STATE:
+
+			EndpointError = EQuicEndpointError::InvalidState;
+			break;
+
+		case QUIC_STATUS_NOT_SUPPORTED:
+
+			EndpointError = EQuicEndpointError::NotSupported;
+			break;
+
+		case QUIC_STATUS_NOT_FOUND:
+
+			EndpointError = EQuicEndpointError::NotFound;
+			break;
+
+		case QUIC_STATUS_BUFFER_TOO_SMALL:
+
+			EndpointError = EQuicEndpointError::BufferTooSmall;
+			break;
+
+		case QUIC_STATUS_HANDSHAKE_FAILURE:
+
+			EndpointError = EQuicEndpointError::ConnectionHandshake;
+			break;
+
+		case QUIC_STATUS_ABORTED:
+
+			EndpointError = EQuicEndpointError::ConnectionAbort;
+			break;
+
+		case QUIC_STATUS_ADDRESS_IN_USE:
+
+			EndpointError = EQuicEndpointError::AddressInUse;
+			break;
+
+		case QUIC_STATUS_INVALID_ADDRESS:
+
+			EndpointError = EQuicEndpointError::InvalidAddress;
+			break;
+
+		case QUIC_STATUS_CONNECTION_TIMEOUT:
+
+			EndpointError = EQuicEndpointError::ConnectionTimeout;
+			break;
+
+		case QUIC_STATUS_CONNECTION_IDLE:
+
+			EndpointError = EQuicEndpointError::ConnectionIdle;
+			break;
+
+		case QUIC_STATUS_INTERNAL_ERROR:
+
+			EndpointError = EQuicEndpointError::InternalError;
+			break;
+
+		case QUIC_STATUS_UNREACHABLE:
+
+			EndpointError = EQuicEndpointError::ServerUnreachable;
+			break;
+
+		case QUIC_STATUS_CONNECTION_REFUSED:
+
+			EndpointError = EQuicEndpointError::ConnectionRefused;
+			break;
+
+		case QUIC_STATUS_PROTOCOL_ERROR:
+
+			EndpointError = EQuicEndpointError::Protocol;
+			break;
+
+		case QUIC_STATUS_VER_NEG_ERROR:
+
+			EndpointError = EQuicEndpointError::VersionNegotiation;
+			break;
+
+		case QUIC_STATUS_USER_CANCELED:
+
+			EndpointError = EQuicEndpointError::UserCanceled;
+			break;
+
+		case QUIC_STATUS_ALPN_NEG_FAILURE:
+
+			EndpointError = EQuicEndpointError::AlpnNegotiation;
+			break;
+
+		case QUIC_STATUS_STREAM_LIMIT_REACHED:
+
+			EndpointError = EQuicEndpointError::StreamLimit;
+			break;
+
+		default:
+
+			EndpointError = EQuicEndpointError::Unknown;
+			break;
+
+		}
+
+		return EndpointError;
+	}
+
+
+	FString GetEndpointErrorString(EQuicEndpointError Error)
+	{
+		FString ErrorMessage;
+
+		switch (Error)
+		{
+
+		case EQuicEndpointError::Normal:
+
+			ErrorMessage = "";
+			break;
+
+		case EQuicEndpointError::OutOfMemory:
+
+			ErrorMessage = "QUIC ran out of memory";
+			break;
+
+		case EQuicEndpointError::InvalidParameter:
+
+			ErrorMessage = "Invalid parameter supplied";
+			break;
+
+		case EQuicEndpointError::InvalidState:
+
+			ErrorMessage = "Invalid state encountered";
+			break;
+
+		case EQuicEndpointError::NotSupported:
+
+			ErrorMessage = "Operation not supported";
+			break;
+
+		case EQuicEndpointError::NotFound:
+
+			ErrorMessage = "Object was not found";
+			break;
+
+		case EQuicEndpointError::BufferTooSmall:
+
+			ErrorMessage = "Buffer is too small";
+			break;
+
+		case EQuicEndpointError::ConnectionHandshake:
+
+			ErrorMessage = "Connection handshake failed";
+			break;
+
+		case EQuicEndpointError::ConnectionAbort:
+
+			ErrorMessage = "Connection was aborted by transport";
+			break;
+
+		case EQuicEndpointError::AddressInUse:
+
+			ErrorMessage = "Client address/port already in use";
+			break;
+
+		case EQuicEndpointError::InvalidAddress:
+
+			ErrorMessage = "Remote address/port invalid";
+			break;
+
+		case EQuicEndpointError::ConnectionTimeout:
+
+			ErrorMessage = "Connection timed out while waiting for a response from peer";
+			break;
+
+		case EQuicEndpointError::ConnectionIdle:
+
+			ErrorMessage = "Connection timed out from inactivity";
+			break;
+
+		case EQuicEndpointError::InternalError:
+
+			ErrorMessage = "Internal Error";
+			break;
+
+		case EQuicEndpointError::ServerUnreachable:
+
+			ErrorMessage = "Server is unreachable";
+			break;
+
+		case EQuicEndpointError::ConnectionRefused:
+
+			ErrorMessage = "Connection was refused";
+			break;
+
+		case EQuicEndpointError::Protocol:
+
+			ErrorMessage = "Encountered protocol error";
+			break;
+
+		case EQuicEndpointError::VersionNegotiation:
+
+			ErrorMessage = "Encountered error during version negotiation";
+			break;
+
+		case EQuicEndpointError::UserCanceled:
+
+			ErrorMessage = "User canceled handshake";
+			break;
+
+		case EQuicEndpointError::AlpnNegotiation:
+
+			ErrorMessage = "Handshake failed to negotiate common ALPN";
+			break;
+
+		case EQuicEndpointError::StreamLimit:
+
+			ErrorMessage = "Stream limit was reached";
+			break;
+
+		case EQuicEndpointError::Unknown:
+
+			ErrorMessage = "Unknown error occurred";
+			break;
+
+		}
+
+		return ErrorMessage;
+	}
+
+}
diff --git a/Engine/Plugins/Experimental/QuicMessaging/Source/QuicMessagingTransport/Private/QuicUtils.h b/Engine/Plugins/Experimental/QuicMessaging/Source/QuicMessagingTransport/Private/QuicUtils.h
new file mode 100644
--- /dev/null
+++ b/Engine/Plugins/Experimental/QuicMessaging/Source/QuicMessagingTransport/Private/QuicUtils.h
@@ -0,0 +1,35 @@
+// Copyright Epic Games, Inc. All Rights Reserved.
+
+#pragma once
+
+#include "CoreMinimal.h"
+#include "QuicFlags.h"
+
+
+namespace QuicUtils
+{
+
+	typedef long HRESULT;
+
+
+	/**
+	 * Convert HRESULT to FString.
+	 *
+	 * @note Most HRESULT values are overriden/extended by msquic
+	 * [REDACTED_PROFILE_URL]/msquic/blob/main/docs/TSG.md#understanding-error-codes
+	 * [REDACTED_PROFILE_URL]/msquic/blob/main/src/inc/msquic_winuser.h
+	 */
+	FString ConvertResult(HRESULT Result);
+
+	/**
+	 * Convert HRESULT msquic status.
+	 */
+	EQuicEndpointError ConvertQuicStatus(HRESULT Status);
+
+	FString GetEndpointErrorString(EQuicEndpointError Error);
+
+};
+
+
+
+
diff --git a/Engine/Plugins/Experimental/QuicMessaging/Source/QuicMessagingTransport/Public/IQuicMessagingTransportModule.h b/Engine/Plugins/Experimental/QuicMessaging/Source/QuicMessagingTransport/Public/IQuicMessagingTransportModule.h
new file mode 100644
--- /dev/null
+++ b/Engine/Plugins/Experimental/QuicMessaging/Source/QuicMessagingTransport/Public/IQuicMessagingTransportModule.h
@@ -0,0 +1,21 @@
+// Copyright Epic Games, Inc. All Rights Reserved.
+
+#pragma once
+
+#include "CoreMinimal.h"
+#include "Modules/ModuleInterface.h"
+#include "Modules/ModuleManager.h"
+
+/**
+ * Interface for QuicMessagingTransport module
+ */
+class IQuicMessagingTransportModule : public IModuleInterface
+{
+public:
+	/** Get the QuicMessagingTransport module */
+	static IQuicMessagingTransportModule& Get()
+	{
+		static const FName ModuleName = TEXT("QuicMessagingTransport");
+		return FModuleManager::Get().GetModuleChecked<IQuicMessagingTransportModule>(ModuleName);
+	}
+};
diff --git a/Engine/Plugins/Experimental/QuicMessaging/Source/QuicMessagingTransport/Public/QuicEndpointConfig.h b/Engine/Plugins/Experimental/QuicMessaging/Source/QuicMessagingTransport/Public/QuicEndpointConfig.h
new file mode 100644
--- /dev/null
+++ b/Engine/Plugins/Experimental/QuicMessaging/Source/QuicMessagingTransport/Public/QuicEndpointConfig.h
@@ -0,0 +1,73 @@
+// Copyright Epic Games, Inc. All Rights Reserved.
+
+#pragma once
+
+struct FIPv4Endpoint;
+struct FGuid;
+enum class EEncryptionMode : uint8;
+enum class EAuthenticationMode : uint8;
+enum class EConnectionCooldownMode : uint8;
+enum class EQuicClientVerificationMode : uint8;
+enum class EQuicClientSelfSignedCertificate : uint8;
+
+
+struct FQuicEndpointConfig
+{
+
+	/** Holds the local/remote endpoint. */
+	FIPv4Endpoint Endpoint = FIPv4Endpoint::Any;
+
+	/** Holds the local node id. */
+	FGuid LocalNodeId = FGuid::NewGuid();
+
+	/** Holds flag indicating whether encryption is enabled. */
+	EEncryptionMode EncryptionMode = EEncryptionMode::Enabled;
+
+	/** Holds the discovery timeout in seconds. */
+	uint32 DiscoveryTimeoutSec = 0;
+
+};
+
+
+struct FQuicServerConfig
+	: FQuicEndpointConfig
+{
+
+	/** Holds the certificate for Quic communication. */
+	FString Certificate = "";
+
+	/** Holds the private key. */
+	FString PrivateKey = "";
+
+	/** Holds the maximum authentication message size. */
+	uint32 MaxAuthenticationMessageSize = 0;
+
+	/** Holds flag indicating whether the server should authenticate clients. */
+	EAuthenticationMode AuthenticationMode = EAuthenticationMode::Disabled;
+
+	/** Holds flag indicating whether there is a cooldown for connection attempts. */
+	EConnectionCooldownMode ConnCooldownMode = EConnectionCooldownMode::Enabled;
+
+	/** Holds the maximum attempts until the connection cooldown is triggered. */
+	uint32 ConnCooldownMaxAttempts = 5;
+
+	/** Holds the time period in seconds within the max attempts must happen. */
+	uint32 ConnCooldownPeriodSec = 30;
+
+	/** Holds the connection cooldown in seconds. */
+	uint32 ConnCooldownSec = 30;
+
+	/** Holds the maximum connection cooldown in seconds. */
+	uint32 ConnCooldownMaxSec = 3600;
+
+};
+
+
+struct FQuicClientConfig
+	: FQuicEndpointConfig
+{
+
+	/** Holds flag indicating whether the client should verify the server certificate. */
+	EQuicClientVerificationMode ClientVerificationMode = EQuicClientVerificationMode::Verify;
+
+};
diff --git a/Engine/Plugins/Experimental/QuicMessaging/Source/QuicMessagingTransport/Public/QuicEndpointManager.h b/Engine/Plugins/Experimental/QuicMessaging/Source/QuicMessagingTransport/Public/QuicEndpointManager.h
new file mode 100644
--- /dev/null
+++ b/Engine/Plugins/Experimental/QuicMessaging/Source/QuicMessagingTransport/Public/QuicEndpointManager.h
@@ -0,0 +1,467 @@
+// Copyright Epic Games, Inc. All Rights Reserved.
+
+#pragma once
+
+#include "CoreMinimal.h"
+#include "CoreTypes.h"
+#include "Containers/Map.h"
+#include "Containers/Ticker.h"
+#include "Interfaces/IPv4/IPv4Endpoint.h"
+#include "Misc/Guid.h"
+#include "Templates/SharedPointer.h"
+#include "INetworkMessagingExtension.h"
+
+#include "QuicFlags.h"
+#include "QuicNodeInfo.h"
+#include "QuicMessages.h"
+
+
+class FQuicEndpoint;
+struct FQuicEndpointConfig;
+struct QUIC_API_TABLE;
+struct QUIC_HANDLE;
+struct QUIC_BUFFER;
+
+typedef QUIC_HANDLE* HQUIC;
+
+
+using FQuicEndpointPtr = TSharedPtr<FQuicEndpoint, ESPMode::ThreadSafe>;
+
+
+/**
+ * QUIC endpoint manager
+ *
+ * Authentication:
+ * - If AuthenticationMode is Enabled, QuicServerHandler will discard
+ *   any message with MessageType other than "Authentication" until
+ *	 the node associated with the handler is marked as authenticated
+ * - Messages with type "Authentication" or "AuthenticationResponse"
+ *   can be sent via the extension methods found in
+ *	 IQuicNetworkMessagingExtension.h
+ * - The MaxAuthenticationMessageSize limit can be set to avoid
+ *   reading potentially large malicious authentication messages
+ *
+ * Client/Server:
+ * - The EndpointManager will start in client mode by default
+ *   and is converted to a server endpoint when SetServerMode
+ *	 is called with additional parameters needed to run a server
+ * - QuicServerHandlers will continuously ping the connected
+ *   client to keep the connection alive
+ */
+class QUICMESSAGINGTRANSPORT_API FQuicEndpointManager
+{
+
+public:
+
+	FQuicEndpointManager() = delete;
+
+	/**
+	 * Initialize endpoint manager with default mode as client.
+	 *
+	 * @param InEndpointConfig The endpoint config
+	 */
+	FQuicEndpointManager(
+		TSharedRef<FQuicEndpointConfig> InEndpointConfig);
+
+	/**
+	 * Default destructor.
+	 */
+	~FQuicEndpointManager();
+
+public:
+
+	/**
+	 * Shutdown endpoint manager and kill client/server threads.
+	 */
+	void Shutdown();
+
+private:
+
+	/**
+	 * Load the client configuration.
+	 */
+	HQUIC LoadClientConfiguration() const;
+
+	/**
+	 * Load the server configuration.
+	 *
+	 * @param CertificatePath Path to the certificate
+	 * @param PrivateKeyPath Path to the private key
+	 */
+	HQUIC LoadServerConfiguration(
+		FString CertificatePath, FString PrivateKeyPath) const;
+
+public:
+
+	/**
+	 * Set endpoint manager to operate in server mode and initialize server.
+
+	 * @see Authentication
+	 */
+	void InitializeServer();
+
+public:
+
+	/**
+	 * Add remote endpoint client.
+	 *
+	 * @param ClientEndpoint The client endpoint that will be connected to
+	 */
+	void AddClient(const FIPv4Endpoint& ClientEndpoint);
+
+	/**
+	 * Remove remote endpoint client.
+	 * 
+	 * @param ClientEndpoint The client endpoint that will be removed
+	 */
+	void RemoveClient(const FIPv4Endpoint& ClientEndpoint);
+
+	/**
+	 * Reconnect remote endpoint client.
+	 *
+	 * @param ClientEndpoint The client endpoint that will be reconnected
+	 */
+	void ReconnectClient(const FIPv4Endpoint ClientEndpoint);
+
+public:
+
+	/**
+	 * Check if a QuicServerHandler endpoint is authenticated.
+	 * 
+	 * @param NodeId The endpoint associated with this node
+	 * @return True if the endpoint is authenticated
+	 * @note This is only used for QuicServers
+	 * @see Authentication
+	 */
+	bool IsEndpointAuthenticated(const FGuid& NodeId) const;
+
+	/**
+	 * Set QuicServerHandler endpoint as authenticated.
+	 * 
+	 * @param NodeId The endpoint associated with this node
+	 * @note This is only used for QuicServers
+	 * @see Authentication
+	 */
+	void SetEndpointAuthenticated(const FGuid& NodeId);
+
+	/**
+	 * Disconnect a node.
+	 *
+	 * @param NodeId The node identifier
+	 */
+	void DisconnectNode(const FGuid& NodeId);
+
+	/**
+	 * Disconnect an endpoint.
+	 *
+	 * @param Endpoint The endpoint address
+	 */
+	void DisconnectEndpoint(const FIPv4Endpoint& Endpoint);
+
+public:
+
+	/**
+	 * Enqueues outbound messages on all recipient endpoints.
+	 * 
+	 * @param Data The message data
+	 * @param MessageMetas The message recipient node identifiers and associated message ID
+	 * @param MessageType The message type
+	 * @see Authentication and EQuicMessageType
+	 */
+	void EnqueueOutboundMessages(FQuicPayloadPtr Data,
+		const TArray<TTuple<FGuid, uint32>>& MessageMetas, const EQuicMessageType MessageType);
+
+	/**
+	 * Delegate for when a message has been delivered.
+	 *
+	 * @return The delegate
+	 */
+	DECLARE_DELEGATE_OneParam(FOnMessageDelivered, const uint32 /*MessageId*/)
+	FOnMessageDelivered& OnMessageDelivered()
+	{
+		return MessageDeliveredDelegate;
+	}
+
+protected:
+
+	/**
+	 * Validates InboundMessage and passes it to MessageProcessor.
+	 * 
+	 * @param InboundMessage The inbound message
+	 */
+	void ValidateInboundMessage(const FInboundMessage InboundMessage);
+
+	/**
+	 * Checks if a message is valid.
+	 *
+	 * @param InboundMessage The inbound message
+	 */
+	bool IsMessageValid(const FInboundMessage InboundMessage);
+
+	/**
+	 * Instructs the appropriate endpoint or handler to send a hello message.
+	 *
+	 * @param InboundMessage The inbound message
+	 */
+	void SendEndpointHello(const FInboundMessage& InboundMessage);
+
+public:
+
+	/**
+	 * Delegate for when a message has been received and validated.
+	 * 
+	 * @return The delegate
+	 */
+	DECLARE_DELEGATE_OneParam(FOnMessageValidated, const FInboundMessage /*InboundMessage*/)
+	FOnMessageValidated& OnMessageValidated()
+	{
+		return MessageValidatedDelegate;
+	}
+
+protected:
+
+	/**
+	 * Forward for the serialize header delegate.
+	 *
+	 * @param MessageHeader The message header
+	 */
+	FQuicPayloadPtr SerializeHeader(
+		FMessageHeader& MessageHeader) const;
+
+	/**
+	 * Forward for the deserialize header delegate.
+	 *
+	 * @param HeaderData The header data
+	 */
+	FMessageHeader DeserializeHeader(
+		const FQuicPayloadPtr HeaderData) const;
+
+public:
+
+	/**
+	 * Delegate for when a header needs to be serialized.
+	 *
+	 * @return The delegate
+	 */
+	DECLARE_DELEGATE_RetVal_OneParam(FQuicPayloadPtr, FOnSerializeHeader, FMessageHeader& /*MessageHeader*/)
+	FOnSerializeHeader& OnSerializeHeader()
+	{
+		return SerializeHeaderDelegate;
+	}
+
+	/**
+	 * Delegate for when a header needs to be deserialized.
+	 *
+	 * @return The delegate
+	 */
+	DECLARE_DELEGATE_RetVal_OneParam(FMessageHeader, FOnDeserializeHeader, FQuicPayloadPtr /*HeaderData*/)
+	FOnDeserializeHeader& OnDeserializeHeader()
+	{
+		return DeserializeHeaderDelegate;
+	}
+
+public:
+
+	/**
+	 * Register discovered endpoint node.
+	 * 
+	 * @param NodeId The node ID
+	 * @param NodeEndpoint The node endpoint address
+	 * @param LocalEndpoint The local endpoint
+	 */
+	void EndpointNodeDiscovered(const FGuid NodeId,
+		const FIPv4Endpoint& NodeEndpoint, const FIPv4Endpoint& LocalEndpoint);
+
+	/**
+	 * De-register lost endpoint node.
+	 * 
+	 * @param NodeId The node ID
+	 * @param LostEndpoint The lost endpoint address
+	 */
+	void EndpointNodeLost(const FGuid NodeId, const FIPv4Endpoint LostEndpoint);
+
+	/**
+	 * Delegate for when a node was discovered.
+	 * 
+	 * @return The delegate
+	 */
+	DECLARE_DELEGATE_OneParam(FOnEndpointNodeDiscovered, const FGuid& /*NodeId*/)
+	FOnEndpointNodeDiscovered& OnEndpointNodeDiscovered()
+	{
+		return EndpointNodeDiscoveredDelegate;
+	}
+	
+	/**
+	 * Delegate for when a node was lost.
+	 * 
+	 * @return The delegate
+	 */
+	DECLARE_DELEGATE_OneParam(FOnEndpointNodeLost, const FGuid& /*NodeId*/)
+	FOnEndpointNodeLost& OnEndpointNodeLost()
+	{
+		return EndpointNodeLostDelegate;
+	}
+
+	/**
+	 * Get the NodeId of a known endpoint.
+	 *
+	 * @param Endpoint The endpoint
+	 *
+	 * @return Optional FGuid if the node was found
+	 */
+	TOptional<FGuid> GetKnownNodeId(const FIPv4Endpoint& Endpoint) const;
+
+	/**
+	 * Get a list of known NodeIds.
+	 */
+	TArray<FGuid> GetKnownNodeIds() const;
+
+	/**
+	 * Check if an endpoint is associated with a known node.
+	 *
+	 * @param Endpoint The endpoint
+	 */
+	bool IsEndpointKnown(const FIPv4Endpoint Endpoint) const;
+
+	/**
+	 * Get a list of the known endpoints.
+	 */
+	TArray<FIPv4Endpoint> GetKnownEndpoints() const;
+
+	/**
+	 * Get the network/transport statistics for a node.
+	 *
+	 * @param NodeId The node ID
+	 */
+	FMessageTransportStatistics GetStats(FGuid NodeId) const;
+
+protected:
+
+	/**
+	 * Removes a disconnected endpoint.
+	 *
+	 * @param DisconnectedEndpoint The disconnected endpoint
+	 * @param Reason The disconnect reason
+	 */
+	void EndpointDisconnected(
+		const FIPv4Endpoint DisconnectedEndpoint, EQuicEndpointError Reason);
+
+	/**
+	 * Checks if the lost endpoint node is still registered and removes it.
+	 *
+	 * @param LostEndpoint The lost endpoint
+	 */
+	void CheckLoseNode(const FIPv4Endpoint LostEndpoint);
+
+	/**
+	 * Adds a DiscoveryTimeout entry if discovery timeouts are enabled.
+	 *
+	 * @param ConnectedEndpoint The connected endpoint
+	 */
+	void EndpointConnected(const FIPv4Endpoint& ConnectedEndpoint);
+
+
+	/**
+	 * Regularly called to check if any of the endpoints have reached
+	 * the discovery timeout. Timed out endpoints will be disconnected.
+	 */
+	void CheckDiscoveryTimeout();
+
+	/**
+	 * Updates a KnownNode with new endpoint transport statistics.
+	 *
+	 * @param TransportStats The endpoint transport statistics
+	 */
+	void EndpointStatisticsUpdated(
+		const FMessageTransportStatistics TransportStats);
+
+private:
+
+	/** Holds the QUIC API table handle. */
+	const QUIC_API_TABLE* MsQuic;
+
+	/** Holds the QUIC API registration handle. */
+	HQUIC Registration;
+
+	/** Holds the current QUIC API configuration handle. */
+	HQUIC CurrentConfiguration;
+
+private:
+
+	/** Holds the local node identifier. */
+	FGuid LocalNodeId;
+
+	/** Holds the local server endpoint that the socket will be bound to. */
+	FIPv4Endpoint ServerEndpoint;
+
+	/** Flag indicating the endpoint mode. */
+	EEndpointMode EndpointMode;
+
+	/** Holds the endpoint config. */
+	TSharedRef<FQuicEndpointConfig> EndpointConfig;
+
+	/** Holds the collection of reconnect attempts. */
+	TMap<FIPv4Endpoint, uint32> ReconnectAttempts;
+
+	/** Holds the collection of active QUIC endpoints. */
+	TMap<FIPv4Endpoint, FQuicEndpointPtr> Endpoints;
+
+	/** Mutex protecting access to the endpoints map. */
+	mutable FCriticalSection EndpointsCS;
+
+	/** Holds the collection of known remote nodes. */
+	TMap<FGuid, FQuicNodeInfo> KnownNodes;
+
+	/** Mutex protecting access to the nodes map. */
+	mutable FCriticalSection NodesCS;
+
+private:
+
+	/** Holds the handle for the discovery timeout tick. */
+	FTSTicker::FDelegateHandle DiscoveryTimeoutTickHandle;
+
+	struct FDiscoveryTimeoutEntry
+	{
+
+		/** Holds the address of the connected endpoint. */
+		FIPv4Endpoint Endpoint = FIPv4Endpoint::Any;
+
+		/** Holds the connection timestamp. */
+		double Timestamp = 0;
+
+	};
+
+	/** Holds the connection timestamps for discovery timeout purposes. */
+	TArray<FDiscoveryTimeoutEntry> DiscoveryTimeouts;
+
+	/** Mutex protecting the discovery timeouts array. */
+	mutable FCriticalSection DiscoveryTimeoutCS;
+
+private:
+
+	/** Holds a delegate to be invoked when a message was delivered to all endpoints. */
+	FOnMessageDelivered MessageDeliveredDelegate;
+
+	/**
+	 * Holds a delegate to be invoked when a message has been received and validated,
+	 * ready to be serialized by the local endpoint.
+	 */
+	FOnMessageValidated MessageValidatedDelegate;
+
+private:
+
+	/** Holds the serialize header delegate. */
+	FOnSerializeHeader SerializeHeaderDelegate;
+
+	/** Holds the deserialize header delegate. */
+	FOnDeserializeHeader DeserializeHeaderDelegate;
+
+private:
+
+	/** Holds a delegate to be invoked when a node was discovered. */
+	FOnEndpointNodeDiscovered EndpointNodeDiscoveredDelegate;
+
+	/** Holds a delegate to be invoked when a node was lost. */
+	FOnEndpointNodeLost EndpointNodeLostDelegate;
+
+};
diff --git a/Engine/Plugins/Experimental/QuicMessaging/Source/QuicMessagingTransport/Public/QuicFlags.h b/Engine/Plugins/Experimental/QuicMessaging/Source/QuicMessagingTransport/Public/QuicFlags.h
new file mode 100644
--- /dev/null
+++ b/Engine/Plugins/Experimental/QuicMessaging/Source/QuicMessagingTransport/Public/QuicFlags.h
@@ -0,0 +1,261 @@
+// Copyright Epic Games, Inc. All Rights Reserved.
+
+#pragma once
+
+
+#include "CoreTypes.h"
+
+
+/**
+ * Enumerates endpoint mode.
+ */
+enum class EEndpointMode : uint8
+{
+	/** EndpointManager acts as a client. */
+	Client,
+
+	/** EndpointManager acts as a server. */
+	Server
+};
+
+
+/**
+ * Enumerates authentication mode.
+ */
+enum class EAuthenticationMode : uint8
+{
+	/** Client nodes have to be authenticated before sending MessageType. */
+	Enabled,
+
+	/** Authentication is disabled, all nodes can connect and send data. */
+	Disabled
+};
+
+
+/**
+ * Enumerates encryption mode.
+ *
+ * When disabled, encryption will only be used for the handshake.
+ * Further communication will be unencrypted.
+ *
+ * @note Both endpoints must have the same setting to successfully connect.
+ */
+enum class EEncryptionMode : uint8
+{
+	/** Encryption is enabled. */
+	Enabled,
+
+	/** Encryption is disabled. */
+	Disabled
+};
+
+
+/**
+ * Enumerates connection cooldown mode.
+ *
+ * When enabled, multiple successive attempts to connect from
+ * the same IP address will trigger a cooldown.
+ *
+ * @note This is a server-only flag
+ */
+enum class EConnectionCooldownMode : uint8
+{
+	/** Connection cooldown is enabled. */
+	Enabled,
+
+	/** Connection cooldown is disabled. */
+	Disabled
+};
+
+
+/**
+ * Enumerates client connection state.
+ */
+enum class EQuicClientState : uint8
+{
+	/** Client is setting up connection. */
+	Connecting,
+
+	/** Client is connected and running. */
+	Connected,
+
+	/** Client could not connect. */
+	Failed,
+
+	/** Client is disconnected. */
+	Disconnected,
+
+	/** Client connection is closed. */
+	Closed,
+
+	/** Client is being destroyed. */
+	Stopped
+};
+
+
+/**
+ * Enumerates client certificate verification.
+ */
+enum class EQuicClientVerificationMode : uint8
+{
+	/** Verify server certificate. */
+	Verify,
+
+	/** Pass server certificate verification. */
+	Pass
+};
+
+
+/**
+ * Enumerates client self signed certificates mode.
+ */
+enum class EQuicClientSelfSignedCertificate : uint8
+{
+	/** Reject self-signed certificates. */
+	Reject,
+
+	/** Allow self-signed certificates. */
+	Allow
+};
+
+
+/**
+ * Enumerates server state.
+ */
+enum class EQuicServerState : uint8
+{
+	/** Server is starting. */
+	Starting,
+
+	/** Server is running. */
+	Running,
+
+	/** Server encountered error and can't recover. */
+	Failed,
+
+	/** Server is stopped. */
+	Stopped
+};
+
+
+/**
+ * Enumerates server handler state.
+ */
+enum class EQuicHandlerState : uint8
+{
+	/** Handler is starting. */
+	Starting,
+
+	/** Handler is running. */
+	Running,
+
+	/** Handler is stopping. */
+	Stopping,
+
+	/** Handler connection failed. */
+	Failed,
+
+	/** Handler has disconnected. */
+	Disconnected
+};
+
+
+/**
+ * Enumerates message type.
+ *
+ * @see QuicEndpointManager - Authentication
+ */
+enum class EQuicMessageType : uint8
+{
+	/** Hello, header-only. */
+	Hello,
+
+	/** Authentication (from client to server) */
+	Authentication,
+
+	/** Authentication response (from server to client) */
+	AuthenticationResponse,
+
+	/** Data */
+	Data,
+
+	/** Bye, header-only. */
+	Bye
+
+};
+
+
+/**
+ * Enumerates endpoint error.
+ */
+enum class EQuicEndpointError : uint16
+{
+	/** No endpoint error. */
+	Normal,
+
+	/** QUIC ran out of memory. */
+	OutOfMemory,
+
+	/** Invalid parameter supplied. */
+	InvalidParameter,
+
+	/** Invalid state encountered. */
+	InvalidState,
+
+	/** Operation not supported. */
+	NotSupported,
+
+	/** Object was not found. */
+	NotFound,
+
+	/** Buffer is too small. */
+	BufferTooSmall,
+
+	/** Connection handshake failed. */
+	ConnectionHandshake,
+
+	/** Connection was aborted by transport. */
+	ConnectionAbort,
+
+	/** Client address/port already in use. */
+	AddressInUse,
+
+	/** Remote address/port invalid. */
+	InvalidAddress,
+
+	/** Connection timed out while waiting for a response from peer. */
+	ConnectionTimeout,
+
+	/** Connection timed out from inactivity. */
+	ConnectionIdle,
+
+	/** Internal error. */
+	InternalError,
+
+	/** Server is unreachable. */
+	ServerUnreachable,
+
+	/** Connection was refused. */
+	ConnectionRefused,
+
+	/** Encountered protocol error. */
+	Protocol,
+
+	/** Encountered error during version negotiation. */
+	VersionNegotiation,
+
+	/** User canceled handshake. */
+	UserCanceled,
+
+	/** Handshake failed to negotiate common ALPN. */
+	AlpnNegotiation,
+
+	/** Stream limit was reached. */
+	StreamLimit,
+
+	/** Unknown error. */
+	Unknown
+
+};
+
+
diff --git a/Engine/Plugins/Experimental/QuicMessaging/Source/QuicMessagingTransport/Public/QuicMessages.h b/Engine/Plugins/Experimental/QuicMessaging/Source/QuicMessagingTransport/Public/QuicMessages.h
new file mode 100644
--- /dev/null
+++ b/Engine/Plugins/Experimental/QuicMessaging/Source/QuicMessagingTransport/Public/QuicMessages.h
@@ -0,0 +1,139 @@
+// Copyright Epic Games, Inc. All Rights Reserved.
+
+#pragma once
+
+#include <memory>
+
+#include "CoreTypes.h"
+#include "Templates/SharedPointer.h"
+#include "GenericPlatform/GenericPlatform.h"
+#include "Interfaces/IPv4/IPv4Endpoint.h"
+#include "Misc/Guid.h"
+
+#include "QuicFlags.h"
+
+
+#define QUIC_MESSAGE_HEADER_SIZE 96
+
+using FQuicPayload = TArray<uint8>;
+using FQuicPayloadPtr = TSharedPtr<FQuicPayload, ESPMode::ThreadSafe>;
+
+
+/**
+ * Structure for QUIC message headers.
+ */
+struct FMessageHeader
+{
+    /** Holds the message type. */
+    EQuicMessageType MessageType;
+
+    /** Holds the recipient node id. */
+    FGuid RecipientId;
+
+    /** Holds the sender node id. */
+    FGuid SenderId;
+
+    /** Holds the serialized message size. */
+    uint32 SerializedMessageSize;
+
+    /** Default constructor. */
+    FMessageHeader()
+	    : MessageType(EQuicMessageType::Hello)
+		, RecipientId(FGuid::NewGuid())
+		, SenderId(FGuid::NewGuid())
+		, SerializedMessageSize(0)
+	{ }
+
+    /** Creates and initializes a new instance. */
+    FMessageHeader(const EQuicMessageType InMessageType, const FGuid InRecipientId,
+        const FGuid InSenderId, const uint32 InSerializedMessageSize)
+        : MessageType(InMessageType)
+        , RecipientId(InRecipientId)
+        , SenderId(InSenderId)
+        , SerializedMessageSize(InSerializedMessageSize)
+    { }
+};
+
+
+/**
+ * Structure for unserialized inbound messages.
+ */
+struct FInboundMessage
+{
+    /** Holds the message data. */
+	FQuicPayloadPtr UnserializedMessage;
+
+    /** Holds the message header. */
+    FMessageHeader MessageHeader;
+
+    /** Holds the sender's network endpoint. */
+    FIPv4Endpoint Sender;
+
+    /** Holds the receiver's network endpoint. */
+    FIPv4Endpoint Receiver;
+
+    /** Default constructor. */
+    FInboundMessage()
+	    : UnserializedMessage(nullptr)
+		, MessageHeader(FMessageHeader())
+		, Sender(FIPv4Endpoint::Any)
+		, Receiver(FIPv4Endpoint::Any)
+	{ }
+
+    /** Creates and initializes a new instance. */
+    FInboundMessage(
+		FQuicPayloadPtr InUnserializedMessage,
+		const FMessageHeader& InMessageHeader,
+		const FIPv4Endpoint& InSender, const FIPv4Endpoint& InReceiver)
+        : UnserializedMessage(InUnserializedMessage)
+        , MessageHeader(InMessageHeader)
+        , Sender(InSender)
+		, Receiver(InReceiver)
+    { }
+};
+
+
+/**
+ * Structure for serialized outbound messages.
+ */
+struct FOutboundMessage
+{
+	/** Holds the recipient endpoint address. */
+	FIPv4Endpoint Recipient;
+
+	/** Flag indicating whether this message has a payload. */
+	bool bHasPayload;
+
+    /** Holds the serialized payload data. */
+	FQuicPayloadPtr SerializedPayload;
+
+	/** Holds the serialized message header. */
+	FQuicPayloadPtr SerializedHeader;
+
+	/** Default constructor. */
+	FOutboundMessage()
+		: Recipient(FIPv4Endpoint::Any)
+		, bHasPayload(false)
+		, SerializedPayload(nullptr)
+		, SerializedHeader(nullptr)
+	{ }
+
+    /** Creates and initializes a new instance. */
+    FOutboundMessage(const FIPv4Endpoint InRecipient,
+		const FQuicPayloadPtr InSerializedMessage,
+		const FQuicPayloadPtr InSerializedHeader)
+	    : Recipient(InRecipient)
+		, bHasPayload(true)
+		, SerializedPayload(InSerializedMessage)
+		, SerializedHeader(InSerializedHeader)
+    { }
+
+	/** Creates and initializes a new header-only instance. */
+	FOutboundMessage(const FIPv4Endpoint InRecipient,
+		const FQuicPayloadPtr InSerializedHeader)
+		: Recipient(InRecipient)
+		, bHasPayload(false)
+		, SerializedPayload(nullptr)
+		, SerializedHeader(InSerializedHeader)
+	{}
+};
diff --git a/Engine/Plugins/Experimental/QuicMessaging/Source/QuicMessagingTransport/Public/QuicNodeInfo.h b/Engine/Plugins/Experimental/QuicMessaging/Source/QuicMessagingTransport/Public/QuicNodeInfo.h
new file mode 100644
--- /dev/null
+++ b/Engine/Plugins/Experimental/QuicMessaging/Source/QuicMessagingTransport/Public/QuicNodeInfo.h
@@ -0,0 +1,48 @@
+// Copyright Epic Games, Inc. All Rights Reserved.
+
+#pragma once
+
+#include "Interfaces/IPv4/IPv4Endpoint.h"
+#include "Misc/Guid.h"
+#include "INetworkMessagingExtension.h"
+
+
+/**
+ * Structure for known remote endpoints.
+ */
+struct FQuicNodeInfo
+{
+	/** Holds the node's endpoint. */
+	FIPv4Endpoint Endpoint;
+
+	/** Holds the endpoint's node identifier. */
+	FGuid NodeId;
+
+	/** Holds the node's local quic endpoint. */
+	FIPv4Endpoint LocalEndpoint;
+
+	/** Denotes wheter this remote node is authenticated or not. */
+	bool bIsAuthenticated;
+
+	/** Various transport statistics for this endpoint */
+	FMessageTransportStatistics Statistics;
+
+	/** Default constructor. */
+	FQuicNodeInfo()
+		: Endpoint(FIPv4Endpoint::Any)
+		, NodeId(FGuid::NewGuid())
+		, LocalEndpoint(FIPv4Endpoint::Any)
+		, bIsAuthenticated(false)
+		, Statistics({})
+	{ }
+
+	/** Creates and initializes a new instance. */
+	FQuicNodeInfo(const FIPv4Endpoint& InEndpoint, const FGuid& InNodeId,
+		const FIPv4Endpoint& InLocalEndpoint)
+		: Endpoint(InEndpoint)
+		, NodeId(InNodeId)
+		, LocalEndpoint(InLocalEndpoint)
+		, bIsAuthenticated(false)
+		, Statistics({})
+	{ }
+};
diff --git a/Engine/Plugins/Experimental/QuicMessaging/Source/QuicMessagingTransport/QuicMessagingTransport.Build.cs b/Engine/Plugins/Experimental/QuicMessaging/Source/QuicMessagingTransport/QuicMessagingTransport.Build.cs
new file mode 100644
--- /dev/null
+++ b/Engine/Plugins/Experimental/QuicMessaging/Source/QuicMessagingTransport/QuicMessagingTransport.Build.cs
@@ -0,0 +1,50 @@
+// Copyright Epic Games, Inc. All Rights Reserved.
+using System.IO;
+
+namespace UnrealBuildTool.Rules
+{
+	public class QuicMessagingTransport : ModuleRules
+	{
+		public QuicMessagingTransport(ReadOnlyTargetRules Target) : base(Target)
+		{
+			bEnableExceptions = true;
+			//bUseRTTI = true;
+
+			if (Target.Platform == UnrealTargetPlatform.Win64) {
+
+				PublicSystemLibraries.AddRange(
+					new string[] {
+						"shlwapi.lib",
+						"crypt32.lib"
+					}
+				);
+
+			}
+			
+			PublicDependencyModuleNames.AddRange(
+				new string[] {
+					"Core",
+					"MsQuic",
+					"MsQuicRuntime"
+				}
+			);
+
+			PrivateDependencyModuleNames.AddRange(
+				new string[] {
+					"Messaging",
+					"Networking",
+					"SSL"
+				}
+			);
+
+			PrivateIncludePaths.AddRange(
+				new string[] {
+					"QuicMessagingTransport/Private",
+					"QuicMessagingTransport/Public",
+				}
+			);
+
+			AddEngineThirdPartyPrivateStaticDependencies(Target, "OpenSSL");
+		}
+	}
+}
