[QuicMessaging] Disabling QuicMessaging platform support for mac

QuicMessaging currently does not work for mac - disabling the platform support for mac.

diff --git a/Engine/Plugins/Experimental/QuicMessaging/QuicMessaging.uplugin b/Engine/Plugins/Experimental/QuicMessaging/QuicMessaging.uplugin
--- a/Engine/Plugins/Experimental/QuicMessaging/QuicMessaging.uplugin
+++ b/Engine/Plugins/Experimental/QuicMessaging/QuicMessaging.uplugin
@@ -16,8 +16,7 @@
 	"Installed": false,
 	"SupportedTargetPlatforms": [
 		"Win64",
-		"Linux",
-		"Mac"
+		"Linux"
 	],
 	"SupportedPrograms": [
 		"UnrealFrontend",
@@ -33,7 +32,7 @@
 			"Name": "QuicMessaging",
 			"Type": "RuntimeAndProgram",
 			"LoadingPhase": "PreDefault",
-			"PlatformAllowList": [ "Win64", "Linux", "Mac" ],
+			"PlatformAllowList": [ "Win64", "Linux"],
 			"ProgramAllowList": [
 				"UnrealFrontend",
 				"UnrealMultiUserServer",
@@ -48,7 +47,7 @@
 			"Name": "QuicMessagingTransport",
 			"Type": "RuntimeAndProgram",
 			"LoadingPhase": "PreDefault",
-			"PlatformAllowList": [ "Win64", "Linux", "Mac" ],
+			"PlatformAllowList": [ "Win64", "Linux" ],
 			"ProgramAllowList": [
 				"UnrealFrontend",
 				"UnrealMultiUserServer",
