[QuicMessaging] Fixing QuicMessaging's static analyzer warnings

MSVC static analyzer turned up some warnings within QuicMessaging, which are addressed with this CL:
- Added missing _Function_class annotations for QUIC_*_CALLBACK functions
- Default initialized QUIC_ADDR
- Changed HRESULT comparison to use QUIC_FAILED instead

diff --git a/Engine/Plugins/Experimental/QuicMessaging/Source/QuicMessagingTransport/Private/QuicClient.cpp b/Engine/Plugins/Experimental/QuicMessaging/Source/QuicMessagingTransport/Private/QuicClient.cpp
--- a/Engine/Plugins/Experimental/QuicMessaging/Source/QuicMessagingTransport/Private/QuicClient.cpp
+++ b/Engine/Plugins/Experimental/QuicMessaging/Source/QuicMessagingTransport/Private/QuicClient.cpp
@@ -192,9 +192,10 @@ void FQuicClient::CollectStatistics()
 	}
 
 	QUIC_STATISTICS Stats{};
+
+	QUIC_STATUS Status;
 	
-	if (const QUIC_STATUS Status
-		= GetConnectionStats(ClientConnection, Stats) != QUIC_STATUS_SUCCESS)
+	if (QUIC_FAILED(Status = GetConnectionStats(ClientConnection, Stats)))
 	{
 		UE_LOG(LogQuicMessagingTransport, Error,
 			TEXT("[QuicClient] Could not collect statistics: %s."),
@@ -251,6 +252,7 @@ void FQuicClient::SendBye()
 }
 
 
+_Function_class_(QUIC_CONNECTION_CALLBACK)
 QUIC_STATUS FQuicClient::ClientConnectionCallback(HQUIC Connection,
 	void* Context, QUIC_CONNECTION_EVENT* Event)
 {
@@ -355,6 +357,7 @@ QUIC_STATUS FQuicClient::ClientConnectionCallback(HQUIC Connection,
 }
 
 
+_Function_class_(QUIC_STREAM_CALLBACK)
 QUIC_STATUS FQuicClient::ClientStreamCallback(HQUIC Stream,
 	void* Context, QUIC_STREAM_EVENT* Event)
 {
diff --git a/Engine/Plugins/Experimental/QuicMessaging/Source/QuicMessagingTransport/Private/QuicClient.h b/Engine/Plugins/Experimental/QuicMessaging/Source/QuicMessagingTransport/Private/QuicClient.h
--- a/Engine/Plugins/Experimental/QuicMessaging/Source/QuicMessagingTransport/Private/QuicClient.h
+++ b/Engine/Plugins/Experimental/QuicMessaging/Source/QuicMessagingTransport/Private/QuicClient.h
@@ -116,6 +116,7 @@ protected:
 	 * @param Context The app context (FQuicServer)
 	 * @param Event The connection event
 	 */
+	_Function_class_(QUIC_CONNECTION_CALLBACK)
 	static QUIC_STATUS ClientConnectionCallback(HQUIC Connection,
 		void* Context, QUIC_CONNECTION_EVENT* Event);
 
@@ -126,6 +127,7 @@ protected:
 	 * @param Context The app context (FQuicServer)
 	 * @param Event The stream event
 	 */
+	_Function_class_(QUIC_STREAM_CALLBACK)
 	static QUIC_STATUS ClientStreamCallback(HQUIC Stream,
 		void* Context, QUIC_STREAM_EVENT* Event);
 
diff --git a/Engine/Plugins/Experimental/QuicMessaging/Source/QuicMessagingTransport/Private/QuicServer.cpp b/Engine/Plugins/Experimental/QuicMessaging/Source/QuicMessagingTransport/Private/QuicServer.cpp
--- a/Engine/Plugins/Experimental/QuicMessaging/Source/QuicMessagingTransport/Private/QuicServer.cpp
+++ b/Engine/Plugins/Experimental/QuicMessaging/Source/QuicMessagingTransport/Private/QuicServer.cpp
@@ -527,6 +527,7 @@ void FQuicServer::SendBye(const FIPv4Endpoint& HandlerEndpoint)
 }
 
 
+_Function_class_(QUIC_LISTENER_CALLBACK)
 QUIC_STATUS FQuicServer::ServerListenerCallback(HQUIC Listener,
 	void* Context, QUIC_LISTENER_EVENT* Event)
 {
@@ -548,6 +549,7 @@ QUIC_STATUS FQuicServer::ServerListenerCallback(HQUIC Listener,
 }
 
 
+_Function_class_(QUIC_CONNECTION_CALLBACK)
 QUIC_STATUS FQuicServer::ServerConnectionCallback(HQUIC Connection,
 	void* Context, QUIC_CONNECTION_EVENT* Event)
 {
@@ -656,6 +658,7 @@ QUIC_STATUS FQuicServer::ServerConnectionCallback(HQUIC Connection,
 }
 
 
+_Function_class_(QUIC_STREAM_CALLBACK)
 QUIC_STATUS FQuicServer::ServerStreamCallback(HQUIC Stream,
 	void* Context, QUIC_STREAM_EVENT* Event)
 {
@@ -1140,7 +1143,8 @@ void FQuicServer::ReachedEndOfStream(const FInboundQuicBuffer& QuicBuffer)
 
 QUIC_ADDR FQuicServer::GetRemoteAddress(HQUIC Connection) const
 {
-	QUIC_ADDR RemoteAddress;
+	QUIC_ADDR RemoteAddress = {{0}};
+
 	uint32_t RemoteAddressSize = sizeof(RemoteAddress);
 
 	if (MsQuic)
diff --git a/Engine/Plugins/Experimental/QuicMessaging/Source/QuicMessagingTransport/Private/QuicServer.h b/Engine/Plugins/Experimental/QuicMessaging/Source/QuicMessagingTransport/Private/QuicServer.h
--- a/Engine/Plugins/Experimental/QuicMessaging/Source/QuicMessagingTransport/Private/QuicServer.h
+++ b/Engine/Plugins/Experimental/QuicMessaging/Source/QuicMessagingTransport/Private/QuicServer.h
@@ -198,6 +198,7 @@ protected:
 	 * @param Context The app context (FQuicServer)
 	 * @param Event The listener event
 	 */
+	_Function_class_(QUIC_LISTENER_CALLBACK)
 	static QUIC_STATUS ServerListenerCallback(HQUIC Listener,
 		void* Context, QUIC_LISTENER_EVENT* Event);
 
@@ -208,6 +209,7 @@ protected:
 	 * @param Context The app context (FQuicServer)
 	 * @param Event The connection event
 	 */
+	_Function_class_(QUIC_CONNECTION_CALLBACK)
 	static QUIC_STATUS ServerConnectionCallback(HQUIC Connection,
 		void* Context, QUIC_CONNECTION_EVENT* Event);
 
@@ -218,6 +220,7 @@ protected:
 	 * @param Context The app context (FQuicServer)
 	 * @param Event The stream event
 	 */
+	_Function_class_(QUIC_STREAM_CALLBACK)
 	static QUIC_STATUS ServerStreamCallback(HQUIC Stream,
 		void* Context, QUIC_STREAM_EVENT* Event);
 
