[QuicMessaging] Adding QuicClient connection changed delegate

Adding a delegate that notifies bound functions when a QuicClient connection has changed (connected, disconnected). This is useful to determine whether a client is ready and can send authentication messages to a QuicServer.

diff --git a/Engine/Plugins/Experimental/QuicMessaging/Source/QuicMessaging/Private/QuicMessagingModule.cpp b/Engine/Plugins/Experimental/QuicMessaging/Source/QuicMessaging/Private/QuicMessagingModule.cpp
--- a/Engine/Plugins/Experimental/QuicMessaging/Source/QuicMessaging/Private/QuicMessagingModule.cpp
+++ b/Engine/Plugins/Experimental/QuicMessaging/Source/QuicMessaging/Private/QuicMessagingModule.cpp
@@ -557,6 +557,19 @@ public:
 		return Transport->OnMetaMessageReceived();
 	}
 
+	virtual FOnQuicClientConnectionChanged& OnQuicClientConnectionChanged() override
+	{
+		const TSharedPtr<FQuicMessageTransport> Transport = WeakBridgeTransport.Pin();
+
+		if (!Transport.IsValid())
+		{
+			UE_LOG(LogQuicMessaging, Fatal,
+				TEXT("[QuicMessagingAPI::OnQuicClientConnectionChanged] Transport pointer not valid."));
+		}
+
+		return Transport->OnClientConnectionChanged();
+	}
+
 protected:
 
 	static bool ParseEndpoint(const FString& InEndpointString,
diff --git a/Engine/Plugins/Experimental/QuicMessaging/Source/QuicMessaging/Private/Transport/QuicMessageTransport.cpp b/Engine/Plugins/Experimental/QuicMessaging/Source/QuicMessaging/Private/Transport/QuicMessageTransport.cpp
--- a/Engine/Plugins/Experimental/QuicMessaging/Source/QuicMessaging/Private/Transport/QuicMessageTransport.cpp
+++ b/Engine/Plugins/Experimental/QuicMessaging/Source/QuicMessaging/Private/Transport/QuicMessageTransport.cpp
@@ -225,6 +225,7 @@ bool FQuicMessageTransport::StartTransport(IMessageTransportHandler& Handler)
 	EndpointManager->OnDeserializeHeader().BindRaw(this, &FQuicMessageTransport::HandleManagerDeserializeHeader);
 	EndpointManager->OnEndpointNodeDiscovered().BindRaw(this, &FQuicMessageTransport::HandleManagerNodeDiscovered);
 	EndpointManager->OnEndpointNodeLost().BindRaw(this, &FQuicMessageTransport::HandleManagerNodeLost);
+	EndpointManager->OnClientConnectionChanged().BindRaw(this, &FQuicMessageTransport::HandleClientConnectionChange);
 
 	if (EndpointMode == EEndpointMode::Server)
 	{
@@ -335,6 +336,12 @@ FOnQuicMetaMessageReceived& FQuicMessageTransport::OnMetaMessageReceived()
 }
 
 
+FOnQuicClientConnectionChanged& FQuicMessageTransport::OnClientConnectionChanged()
+{
+	return OnQuicClientConnectionChangedDelegate;
+}
+
+
 void FQuicMessageTransport::HandleProcessorMessageDeserialized(
 	const TSharedPtr<FQuicDeserializedMessage, ESPMode::ThreadSafe> DeserializedMessage,
 	const FGuid NodeId) const
@@ -453,3 +460,13 @@ void FQuicMessageTransport::HandleManagerNodeLost(const FGuid& LostNodeId)
 }
 
 
+void FQuicMessageTransport::HandleClientConnectionChange(
+	const FGuid& NodeId, const FIPv4Endpoint& RemoteEndpoint,
+	const EQuicClientConnectionChange ConnectionState)
+{
+	EQuicClientConnectionState ClientState
+		= (ConnectionState == EQuicClientConnectionChange::Connected)
+		? EQuicClientConnectionState::Connected : EQuicClientConnectionState::Disconnected;
+
+	OnQuicClientConnectionChangedDelegate.Broadcast(NodeId, RemoteEndpoint, ClientState);
+}
diff --git a/Engine/Plugins/Experimental/QuicMessaging/Source/QuicMessaging/Private/Transport/QuicMessageTransport.h b/Engine/Plugins/Experimental/QuicMessaging/Source/QuicMessaging/Private/Transport/QuicMessageTransport.h
--- a/Engine/Plugins/Experimental/QuicMessaging/Source/QuicMessaging/Private/Transport/QuicMessageTransport.h
+++ b/Engine/Plugins/Experimental/QuicMessaging/Source/QuicMessaging/Private/Transport/QuicMessageTransport.h
@@ -17,6 +17,7 @@
 #include "QuicFlags.h"
 #include "QuicMessages.h"
 #include "QuicTransportMessages.h"
+#include "QuicTransportNotifications.h"
 
 
 class FArrayReader;
@@ -212,6 +213,12 @@ public:
 	 */
 	FOnQuicMetaMessageReceived& OnMetaMessageReceived();
 
+	/**
+	 * Delegate that notifies bound functions when a client
+	 * has connected to or disconnected from a remote endpoint (QuicServer).
+	 */
+	FOnQuicClientConnectionChanged& OnClientConnectionChanged();
+
 private:
 
 	/**
@@ -259,6 +266,13 @@ private:
 	 */
 	void HandleManagerNodeLost(const FGuid& LostNodeId);
 
+	/**
+	 * Handles client connection changes.
+	 */
+	void HandleClientConnectionChange(
+		const FGuid& NodeId, const FIPv4Endpoint& RemoteEndpoint,
+		const EQuicClientConnectionChange ConnectionState);
+
 private:
 
 	/** Holds any pending restart requests. */
@@ -296,4 +310,7 @@ private:
 	/** Holds the meta message received delegate. */
 	FOnQuicMetaMessageReceived OnQuicMetaMessageDelegate;
 
+	/** Holds the client connection changed delegate. */
+	FOnQuicClientConnectionChanged OnQuicClientConnectionChangedDelegate;
+
 };
diff --git a/Engine/Plugins/Experimental/QuicMessaging/Source/QuicMessaging/Public/IQuicNetworkMessagingExtension.h b/Engine/Plugins/Experimental/QuicMessaging/Source/QuicMessaging/Public/IQuicNetworkMessagingExtension.h
--- a/Engine/Plugins/Experimental/QuicMessaging/Source/QuicMessaging/Public/IQuicNetworkMessagingExtension.h
+++ b/Engine/Plugins/Experimental/QuicMessaging/Source/QuicMessaging/Public/IQuicNetworkMessagingExtension.h
@@ -6,6 +6,7 @@
 #include "INetworkMessagingExtension.h"
 
 #include "QuicTransportMessages.h"
+#include "QuicTransportNotifications.h"
 
 
 /**
@@ -108,4 +109,12 @@ public:
 	 */
 	virtual FOnQuicMetaMessageReceived& OnQuicMetaMessageReceived() = 0;
 
+	/**
+	 * Delegate that notifies bound functions when a QuicClient
+	 * has connected to or disconnected from a remote endpoint (QuicServer).
+	 *
+	 * @return The delegate
+	 */
+	virtual FOnQuicClientConnectionChanged& OnQuicClientConnectionChanged() = 0;
+
 };
diff --git a/Engine/Plugins/Experimental/QuicMessaging/Source/QuicMessaging/Public/QuicTransportNotifications.h b/Engine/Plugins/Experimental/QuicMessaging/Source/QuicMessaging/Public/QuicTransportNotifications.h
new file mode 100644
--- /dev/null
+++ b/Engine/Plugins/Experimental/QuicMessaging/Source/QuicMessaging/Public/QuicTransportNotifications.h
@@ -0,0 +1,31 @@
+// Copyright Epic Games, Inc. All Rights Reserved.
+
+#pragma once
+
+#include "CoreTypes.h"
+
+struct FGuid;
+struct FIPv4Endpoint;
+
+
+/**
+ * Enumerates QuicClient connection state to the remote endpoint.
+ */
+enum class EQuicClientConnectionState : uint8
+{
+	/** Client connected to remote endpoint. */
+	Connected,
+
+	/** Client disconnected from remote endpoint. */
+	Disconnected
+};
+
+
+/**
+ * Delegate informing bound functions when a QuicClient
+ * has connected to or disconnected from a remote endpoint (QuicServer).
+ */
+DECLARE_MULTICAST_DELEGATE_ThreeParams(FOnQuicClientConnectionChanged,
+	const FGuid& /*NodeId*/, const FIPv4Endpoint& /*RemoteEndpoint*/,
+	const EQuicClientConnectionState /*ConnectionState*/);
+
diff --git a/Engine/Plugins/Experimental/QuicMessaging/Source/QuicMessagingTransport/Private/QuicEndpointManager.cpp b/Engine/Plugins/Experimental/QuicMessaging/Source/QuicMessagingTransport/Private/QuicEndpointManager.cpp
--- a/Engine/Plugins/Experimental/QuicMessaging/Source/QuicMessagingTransport/Private/QuicEndpointManager.cpp
+++ b/Engine/Plugins/Experimental/QuicMessaging/Source/QuicMessagingTransport/Private/QuicEndpointManager.cpp
@@ -798,6 +798,12 @@ void FQuicEndpointManager::EndpointNodeDiscovered(const FGuid NodeId,
 	}
 
     EndpointNodeDiscoveredDelegate.ExecuteIfBound(NodeId);
+
+	if (EndpointMode == EEndpointMode::Client)
+	{
+		ClientConnectionChangedDelegate.ExecuteIfBound(
+			NodeId, NodeEndpoint, EQuicClientConnectionChange::Connected);
+	}
 }
 
 
@@ -825,6 +831,12 @@ void FQuicEndpointManager::EndpointNodeLost(const FGuid NodeId, const FIPv4Endpo
 	}
 
     EndpointNodeLostDelegate.ExecuteIfBound(NodeId);
+
+	if (EndpointMode == EEndpointMode::Client)
+	{
+		ClientConnectionChangedDelegate.ExecuteIfBound(
+			NodeId, LostEndpoint, EQuicClientConnectionChange::Disconnected);
+	}
 }
 
 
diff --git a/Engine/Plugins/Experimental/QuicMessaging/Source/QuicMessagingTransport/Public/QuicEndpointManager.h b/Engine/Plugins/Experimental/QuicMessaging/Source/QuicMessagingTransport/Public/QuicEndpointManager.h
--- a/Engine/Plugins/Experimental/QuicMessaging/Source/QuicMessagingTransport/Public/QuicEndpointManager.h
+++ b/Engine/Plugins/Experimental/QuicMessaging/Source/QuicMessagingTransport/Public/QuicEndpointManager.h
@@ -15,7 +15,6 @@
 #include "QuicNodeInfo.h"
 #include "QuicMessages.h"
 
-
 class FQuicEndpoint;
 struct FQuicEndpointConfig;
 struct QUIC_API_TABLE;
@@ -302,6 +301,19 @@ public:
 		return EndpointNodeLostDelegate;
 	}
 
+	/**
+	 * Delegate for when a client connection changes.
+	 *
+	 * @return The delegate
+	 */
+	DECLARE_DELEGATE_ThreeParams(FOnClientConnectionChanged,
+		const FGuid& /*NodeId*/, const FIPv4Endpoint& /*RemoteEndpoint*/,
+		const EQuicClientConnectionChange /*ConnectionState*/)
+	FOnClientConnectionChanged& OnClientConnectionChanged()
+	{
+		return ClientConnectionChangedDelegate;
+	}
+
 	/**
 	 * Get the NodeId of a known endpoint.
 	 *
@@ -464,4 +476,7 @@ private:
 	/** Holds a delegate to be invoked when a node was lost. */
 	FOnEndpointNodeLost EndpointNodeLostDelegate;
 
+	/** Holds a delegate to be invoked when a client connection has changed. */
+	FOnClientConnectionChanged ClientConnectionChangedDelegate;
+
 };
diff --git a/Engine/Plugins/Experimental/QuicMessaging/Source/QuicMessagingTransport/Public/QuicFlags.h b/Engine/Plugins/Experimental/QuicMessaging/Source/QuicMessagingTransport/Public/QuicFlags.h
--- a/Engine/Plugins/Experimental/QuicMessaging/Source/QuicMessagingTransport/Public/QuicFlags.h
+++ b/Engine/Plugins/Experimental/QuicMessaging/Source/QuicMessagingTransport/Public/QuicFlags.h
@@ -93,6 +93,19 @@ enum class EQuicClientState : uint8
 };
 
 
+/**
+ * Enumerates client connection changes.
+ */
+enum class EQuicClientConnectionChange : uint8
+{
+	/** Client has connected. */
+	Connected,
+
+	/** Client has disconnected. */
+	Disconnected
+};
+
+
 /**
  * Enumerates client certificate verification.
  */
