[QuicMessaging] Adding Mac to support target platforms list to fix shipping build

diff --git a/Engine/Plugins/Experimental/QuicMessaging/QuicMessaging.uplugin b/Engine/Plugins/Experimental/QuicMessaging/QuicMessaging.uplugin
--- a/Engine/Plugins/Experimental/QuicMessaging/QuicMessaging.uplugin
+++ b/Engine/Plugins/Experimental/QuicMessaging/QuicMessaging.uplugin
@@ -16,7 +16,8 @@
 	"Installed": false,
 	"SupportedTargetPlatforms": [
 		"Win64",
-		"Linux"
+		"Linux",
+		"Mac"
 	],
 	"SupportedPrograms": [
 		"UnrealFrontend",
@@ -32,7 +33,7 @@
 			"Name": "QuicMessaging",
 			"Type": "RuntimeAndProgram",
 			"LoadingPhase": "PreDefault",
-			"PlatformAllowList": [ "Win64", "Linux"],
+			"PlatformAllowList": [ "Win64", "Linux" ],
 			"ProgramAllowList": [
 				"UnrealFrontend",
 				"UnrealMultiUserServer",
