[QuicMessaging] Prevent MessageProcessor from processing messages until PostDefault plugin loading phase

diff --git a/Engine/Plugins/Experimental/QuicMessaging/Source/QuicMessaging/Private/Transport/QuicMessageProcessor.cpp b/Engine/Plugins/Experimental/QuicMessaging/Source/QuicMessaging/Private/Transport/QuicMessageProcessor.cpp
--- a/Engine/Plugins/Experimental/QuicMessaging/Source/QuicMessaging/Private/Transport/QuicMessageProcessor.cpp
+++ b/Engine/Plugins/Experimental/QuicMessaging/Source/QuicMessaging/Private/Transport/QuicMessageProcessor.cpp
@@ -10,6 +10,7 @@
 #include "Serialization/ArrayWriter.h"
 #include "Math/UnrealMathUtility.h"
 #include "UObject/Class.h"
+#include "Interfaces/IPluginManager.h"
 
 #include "Backends/CborStructSerializerBackend.h"
 #include "StructSerializer.h"
@@ -67,8 +68,17 @@ FQuicMessageProcessor::FQuicMessageProcessor()
         FPlatformProcess::ReturnSynchEventToPool(EventToDelete);
     });
 
-    Thread = FRunnableThread::Create(this, *FString::Printf(TEXT("FQuicMessageProcessor-%d"), FMath::RandRange(0, 10)), 128 * 4096,
-        TPri_AboveNormal, FPlatformAffinity::GetPoolThreadMask());
+	const ELoadingPhase::Type LoadingPhase = IPluginManager::Get().GetLastCompletedLoadingPhase();
+
+	if (LoadingPhase == ELoadingPhase::PostDefault || LoadingPhase == ELoadingPhase::PostEngineInit)
+	{
+		StartThread();
+	}
+	else
+	{
+		IPluginManager::Get().OnLoadingPhaseComplete().AddRaw(
+			this, &FQuicMessageProcessor::OnPluginLoadingPhaseComplete);
+	}
 
 	UE_LOG(LogQuicMessaging, Verbose, TEXT("[MessageProcessor] Started message processor."));
 }
@@ -76,9 +86,33 @@ FQuicMessageProcessor::FQuicMessageProcessor()
 
 FQuicMessageProcessor::~FQuicMessageProcessor()
 {
-    Thread->Kill(true);
-    delete Thread;
-    Thread = nullptr;
+	if (Thread.IsValid())
+	{
+		Thread->Kill();
+		Thread.Reset();
+	}
+}
+
+
+void FQuicMessageProcessor::StartThread()
+{
+	const uint32 ProcessorThreadStackSize = 128 * 4096;
+
+	Thread = TUniquePtr<FRunnableThread>(FRunnableThread::Create(this, TEXT("FQuicMessageProcessor"),
+		ProcessorThreadStackSize, TPri_AboveNormal, FPlatformAffinity::GetPoolThreadMask()));
+}
+
+
+void FQuicMessageProcessor::OnPluginLoadingPhaseComplete(
+	ELoadingPhase::Type LoadingPhase, bool bPhaseSuccessful)
+{
+	check(!Thread);
+
+	if (LoadingPhase == ELoadingPhase::PostDefault)
+	{
+		IPluginManager::Get().OnLoadingPhaseComplete().RemoveAll(this);
+		StartThread();
+	}
 }
 
 
diff --git a/Engine/Plugins/Experimental/QuicMessaging/Source/QuicMessaging/Private/Transport/QuicMessageProcessor.h b/Engine/Plugins/Experimental/QuicMessaging/Source/QuicMessaging/Private/Transport/QuicMessageProcessor.h
--- a/Engine/Plugins/Experimental/QuicMessaging/Source/QuicMessaging/Private/Transport/QuicMessageProcessor.h
+++ b/Engine/Plugins/Experimental/QuicMessaging/Source/QuicMessaging/Private/Transport/QuicMessageProcessor.h
@@ -6,6 +6,7 @@
 #include "CoreTypes.h"
 #include "Containers/Map.h"
 #include "Containers/Queue.h"
+#include "ModuleDescriptor.h"
 #include "HAL/Runnable.h"
 #include "IMessageTransport.h"
 #include "Interfaces/IPv4/IPv4Endpoint.h"
@@ -76,6 +77,15 @@ public:
 	/** Virtual destructor. */
 	virtual ~FQuicMessageProcessor();
 
+private:
+
+	/** Start the message processor thread. */
+	void StartThread();
+
+	/** Delegate invoked when a plugin phase has been completed. */
+	void OnPluginLoadingPhaseComplete(
+		ELoadingPhase::Type LoadingPhase, bool bPhaseSuccessful);
+
 public:
 
 	/**
@@ -189,7 +199,7 @@ private:
 	EQuicMessageFormat MessageFormat;
 
 	/** Holds the thread object. */
-	FRunnableThread* Thread;
+	TUniquePtr<FRunnableThread> Thread;
 
 	/** Holds an event signaling that outbound messages need to be processed. */
 	TSharedPtr<FEvent, ESPMode::ThreadSafe> SerializeWorkEvent;
